Unity build CIS fixes.

diff --git a/Engine/Plugins/Experimental/AlembicImporter/Source/AlembicLibrary/Private/AbcImportUtilities.cpp b/Engine/Plugins/Experimental/AlembicImporter/Source/AlembicLibrary/Private/AbcImportUtilities.cpp
--- a/Engine/Plugins/Experimental/AlembicImporter/Source/AlembicLibrary/Private/AbcImportUtilities.cpp
+++ b/Engine/Plugins/Experimental/AlembicImporter/Source/AlembicLibrary/Private/AbcImportUtilities.cpp
@@ -3,7 +3,17 @@
 #pragma once
 
 #include "AlembicLibraryPublicPCH.h"
+
+#include <Alembic/AbcCoreHDF5/All.h>
+#include <Alembic/AbcCoreOgawa/All.h>
+#include <Alembic/AbcCoreFactory/All.h>
+#include <Alembic/AbcCoreAbstract/TimeSampling.h>
+#include <Alembic/AbcCoreHDF5/All.h>
+#include <Alembic/Abc/All.h>
+#include <Alembic/AbcGeom/All.h>
+
 #include "AbcImportUtilities.h"
+#include "AbcImportData.h"
 
 #define LOCTEXT_NAMESPACE "AbcImporterUtilities"
 
diff --git a/Engine/Plugins/Experimental/AlembicImporter/Source/AlembicLibrary/Private/AbcImportUtilities.h b/Engine/Plugins/Experimental/AlembicImporter/Source/AlembicLibrary/Private/AbcImportUtilities.h
--- a/Engine/Plugins/Experimental/AlembicImporter/Source/AlembicLibrary/Private/AbcImportUtilities.h
+++ b/Engine/Plugins/Experimental/AlembicImporter/Source/AlembicLibrary/Private/AbcImportUtilities.h
@@ -4,6 +4,10 @@
 
 #include "AlembicLibraryPublicPCH.h"
 
+#include <Alembic/AbcCoreAbstract/TimeSampling.h>
+#include <Alembic/Abc/All.h>
+#include <Alembic/AbcGeom/All.h>
+
 #include "GeometryCache.h"
 #include "GeometryCacheTrackFlipbookAnimation.h"
 #include "GeometryCacheTrackTransformAnimation.h"
@@ -17,6 +21,10 @@
 #include "AbcImportLogger.h"
 #include "AbcImportSettings.h"
 
+struct FAbcMeshSample;
+struct FAbcPolyMeshObject;
+struct FCompressedAbcData;
+
 namespace AbcImporterUtilities
 {
 	/** Templated function to check whether or not an object is of a certain type */
@@ -191,4 +199,4 @@ namespace AbcImporterUtilities
 	void CalculateNewStartAndEndFrameIndices(const float FrameStepRatio, uint32& InOutStartFrameIndex, uint32& InOutEndFrameIndex );
 
 	bool AreVerticesEqual(const FSoftSkinVertex& V1, const FSoftSkinVertex& V2);
-};
+}
