Fix non-unity warning due to redefinition of macro.

diff --git a/Engine/Plugins/Experimental/AlembicImporter/Source/AlembicLibrary/Private/EigenHelper.h b/Engine/Plugins/Experimental/AlembicImporter/Source/AlembicLibrary/Private/EigenHelper.h
--- a/Engine/Plugins/Experimental/AlembicImporter/Source/AlembicLibrary/Private/EigenHelper.h
+++ b/Engine/Plugins/Experimental/AlembicImporter/Source/AlembicLibrary/Private/EigenHelper.h
@@ -3,7 +3,9 @@
 #pragma once
 
 // Just to be sure, also added this in Eigen.Build.cs
+#ifndef EIGEN_MPL2_ONLY
 #define EIGEN_MPL2_ONLY
+#endif
 
 #if defined(_MSC_VER) && USING_CODE_ANALYSIS
 	#pragma warning(push)
