- Remove #pragma once from a .cpp file (doesn't show up in unity builds because .cpp files are #included, heh)

diff --git a/Engine/Plugins/Experimental/AlembicImporter/Source/AlembicLibrary/Private/AbcImportUtilities.cpp b/Engine/Plugins/Experimental/AlembicImporter/Source/AlembicLibrary/Private/AbcImportUtilities.cpp
--- a/Engine/Plugins/Experimental/AlembicImporter/Source/AlembicLibrary/Private/AbcImportUtilities.cpp
+++ b/Engine/Plugins/Experimental/AlembicImporter/Source/AlembicLibrary/Private/AbcImportUtilities.cpp
@@ -1,7 +1,5 @@
 // Copyright 1998-2016 Epic Games, Inc. All Rights Reserved.
 
-#pragma once
-
 #include "AlembicLibraryPublicPCH.h"
 
 #include <Alembic/AbcCoreHDF5/All.h>
