Crash importing specific Alembic file
#fix Additional fix to triangulation vs indexing of vertex attributes, this fixes the crash but also normal issues seen before
#misc default constructor for FTrackRenderData which could cause jittering/invisible geometry

diff --git a/Engine/Plugins/Experimental/AlembicImporter/Source/AlembicLibrary/Private/AbcImportUtilities.cpp b/Engine/Plugins/Experimental/AlembicImporter/Source/AlembicLibrary/Private/AbcImportUtilities.cpp
--- a/Engine/Plugins/Experimental/AlembicImporter/Source/AlembicLibrary/Private/AbcImportUtilities.cpp
+++ b/Engine/Plugins/Experimental/AlembicImporter/Source/AlembicLibrary/Private/AbcImportUtilities.cpp
@@ -445,7 +445,7 @@ bool AbcImporterUtilities::GenerateAbcMeshSampleDataForFrame(const Alembic::AbcG
 		Alembic::AbcGeom::IV2fGeomParam UVCoordinateParameter = Schema.getUVsParam();
 		if (UVCoordinateParameter.valid())
 		{
-			ReadUVSetData(UVCoordinateParameter, FrameSelector, Sample->UVs[0], Sample->Indices, bNeedsTriangulation, FaceCounts);
+			ReadUVSetData(UVCoordinateParameter, FrameSelector, Sample->UVs[0], Sample->Indices, bNeedsTriangulation, FaceCounts, Sample->Vertices.Num());
 		}
 		else
 		{
@@ -462,7 +462,7 @@ bool AbcImporterUtilities::GenerateAbcMeshSampleDataForFrame(const Alembic::AbcG
 				if (Alembic::AbcGeom::IV2fGeomParam::matches(PropertyHeader))
 				{
 					UVSetProperty = Alembic::AbcGeom::IV2fGeomParam(GeomParams, PropertyHeader.getName());
-					ReadUVSetData(UVSetProperty, FrameSelector, Sample->UVs[Sample->NumUVSets], Sample->Indices, bNeedsTriangulation, FaceCounts);
+					ReadUVSetData(UVSetProperty, FrameSelector, Sample->UVs[Sample->NumUVSets], Sample->Indices, bNeedsTriangulation, FaceCounts, Sample->Vertices.Num());
 					++Sample->NumUVSets;
 				}
 			}
@@ -500,16 +500,7 @@ bool AbcImporterUtilities::GenerateAbcMeshSampleDataForFrame(const Alembic::AbcG
 			}
 			else
 			{
-				if (bNeedsTriangulation)
-				{
-					TriangulateVertexAttributeBuffer(FaceCounts, Sample->Normals);
-				}
-
-				// For vertex only normals (and no normal indices available), expand using the regular indices
-				if (Sample->Normals.Num() != Sample->Indices.Num())
-				{
-					ExpandVertexAttributeArray<FVector>(Sample->Indices, Sample->Normals);
-				}
+				ProcessVertexAttributeArray(Sample->Indices, FaceCounts, bNeedsTriangulation, Sample->Vertices.Num(), Sample->Normals);
 			}
 		}
 	}
@@ -574,15 +565,7 @@ bool AbcImporterUtilities::GenerateAbcMeshSampleDataForFrame(const Alembic::AbcG
 			}
 			else
 			{
-				if (bNeedsTriangulation)
-				{
-					TriangulateVertexAttributeBuffer(FaceCounts, Sample->Colors);
-				}
-
-				if (Sample->Colors.Num() != Sample->Indices.Num())
-				{
-					ExpandVertexAttributeArray<FLinearColor>(Sample->Indices, Sample->Colors);
-				}
+				ProcessVertexAttributeArray(Sample->Indices, FaceCounts, bNeedsTriangulation, Sample->Vertices.Num(), Sample->Colors);
 			}
 		}
 		else if (Color4Property.valid())
@@ -609,15 +592,7 @@ bool AbcImporterUtilities::GenerateAbcMeshSampleDataForFrame(const Alembic::AbcG
 			}
 			else
 			{
-				if (bNeedsTriangulation)
-				{
-					TriangulateVertexAttributeBuffer(FaceCounts, Sample->Colors);
-				}
-
-				if (Sample->Colors.Num() != Sample->Indices.Num())
-				{
-					ExpandVertexAttributeArray<FLinearColor>(Sample->Indices, Sample->Colors);
-				}
+				ProcessVertexAttributeArray(Sample->Indices, FaceCounts, bNeedsTriangulation, Sample->Vertices.Num(), Sample->Colors);
 			}
 		}
 		else
@@ -656,7 +631,7 @@ bool AbcImporterUtilities::GenerateAbcMeshSampleDataForFrame(const Alembic::AbcG
 	return bRetrievalResult;
 }
 
-void AbcImporterUtilities::ReadUVSetData(Alembic::AbcGeom::IV2fGeomParam &UVCoordinateParameter, const Alembic::Abc::ISampleSelector FrameSelector, TArray<FVector2D>& OutUVs, const TArray<uint32>& MeshIndices, const bool bNeedsTriangulation, const TArray<uint32>& FaceCounts)
+void AbcImporterUtilities::ReadUVSetData(Alembic::AbcGeom::IV2fGeomParam &UVCoordinateParameter, const Alembic::Abc::ISampleSelector FrameSelector, TArray<FVector2D>& OutUVs, const TArray<uint32>& MeshIndices, const bool bNeedsTriangulation, const TArray<uint32>& FaceCounts, const int32 NumVertices)
 {
 	Alembic::Abc::V2fArraySamplePtr UVSample = UVCoordinateParameter.getValueProperty().getValue(FrameSelector);
 	RetrieveTypedAbcData<Alembic::Abc::V2fArraySamplePtr, FVector2D>(UVSample, OutUVs);
@@ -679,15 +654,7 @@ void AbcImporterUtilities::ReadUVSetData(Alembic::AbcGeom::IV2fGeomParam &UVCoor
 	}
 	else if (OutUVs.Num())
 	{
-		// For vertex only normals (and no normal indices available), expand using the regular indices
-		if (OutUVs.Num() != MeshIndices.Num())
-		{
-			ExpandVertexAttributeArray<FVector2D>(MeshIndices, OutUVs);
-		}
-		else if (bNeedsTriangulation)
-		{
-			TriangulateVertexAttributeBuffer(FaceCounts, OutUVs);
-		}
+		ProcessVertexAttributeArray(MeshIndices, FaceCounts, bNeedsTriangulation, NumVertices, OutUVs);
 	}
 }
 
diff --git a/Engine/Plugins/Experimental/AlembicImporter/Source/AlembicLibrary/Private/AbcImportUtilities.h b/Engine/Plugins/Experimental/AlembicImporter/Source/AlembicLibrary/Private/AbcImportUtilities.h
--- a/Engine/Plugins/Experimental/AlembicImporter/Source/AlembicLibrary/Private/AbcImportUtilities.h
+++ b/Engine/Plugins/Experimental/AlembicImporter/Source/AlembicLibrary/Private/AbcImportUtilities.h
@@ -142,6 +142,20 @@ namespace AbcImporterUtilities
 		// Set new data
 		InOutData = NewData;
 	}
+
+	template<typename T> void ProcessVertexAttributeArray(const TArray<uint32>& InIndices, const TArray<uint32>& InFaceCounts, const bool bTriangulation, const uint32 NumVertices, TArray<T>& InOutArray)
+	{
+		// Expand using the vertex indices (if num entries == num vertices)
+		if (InOutArray.Num() != InIndices.Num() && InOutArray.Num() == NumVertices)
+		{
+			ExpandVertexAttributeArray(InIndices, InOutArray);
+		}
+		// Otherwise the attributes are stored per face, so triangulate if the faces contain quads
+		else if (bTriangulation)
+		{
+			TriangulateVertexAttributeBuffer(InFaceCounts, InOutArray);
+		}
+	}
 	
 	/** Triangulates material indices according to the face counts (quads will have to be split up into two faces / material indices)*/
 	void TriangulateMaterialIndices(const TArray<uint32>& InFaceCounts, TArray<int32>& InOutData);
@@ -168,7 +182,7 @@ namespace AbcImporterUtilities
 	bool GenerateAbcMeshSampleDataForFrame(const Alembic::AbcGeom::IPolyMeshSchema &Schema, const Alembic::Abc::ISampleSelector FrameSelector, FAbcMeshSample* &Sample, const ESampleReadFlags ReadFlags, const bool bFirstFrame );
 
 	/** Read out texture coordinate data from Alembic GeometryParameter */
-	void ReadUVSetData(Alembic::AbcGeom::IV2fGeomParam &UVCoordinateParameter, const Alembic::Abc::ISampleSelector FrameSelector, TArray<FVector2D>& OutUVs, const TArray<uint32>& MeshIndices, const bool bNeedsTriangulation, const TArray<uint32>& FaceCounts);
+	void ReadUVSetData(Alembic::AbcGeom::IV2fGeomParam &UVCoordinateParameter, const Alembic::Abc::ISampleSelector FrameSelector, TArray<FVector2D>& OutUVs, const TArray<uint32>& MeshIndices, const bool bNeedsTriangulation, const TArray<uint32>& FaceCounts, const int32 NumVertices);
 
 	void GenerateSmoothingGroupsIndices(FAbcMeshSample* MeshSample, float HardEdgeAngleThreshold);
 
diff --git a/Engine/Plugins/Experimental/GeometryCache/Source/GeometryCache/Classes/GeometryCacheComponent.h b/Engine/Plugins/Experimental/GeometryCache/Source/GeometryCache/Classes/GeometryCacheComponent.h
--- a/Engine/Plugins/Experimental/GeometryCache/Source/GeometryCache/Classes/GeometryCacheComponent.h
+++ b/Engine/Plugins/Experimental/GeometryCache/Source/GeometryCache/Classes/GeometryCacheComponent.h
@@ -17,6 +17,8 @@ struct FTrackRenderData
 {
 	GENERATED_USTRUCT_BODY()
 
+	FTrackRenderData() : Matrix(FMatrix::Identity), BoundingBox(EForceInit::ForceInitToZero), MatrixSampleIndex(INDEX_NONE), BoundsSampleIndex(INDEX_NONE) {}
+
 	/** Transform matrix used to render this specific track. 
 		This goes from track local space to component local space.
 	*/
