EngineTest crashes with Assertion failed: AbcFile->GetNumPolyMeshes() > 0
#fix Change the way we cosntruct file path for Alembic source (import) files, which was causing issues with seperate engine binaries combined with EngineTest
#fix Do the actual import during PrepareTest rather than BeginPlay, this required exposing SetMobility to BPs on StaticMeshActor otherwise we can't set StaticMesh on freshly spawned UStaticMeshActor

diff --git a/Engine/Plugins/Experimental/AlembicImporter/Source/AlembicLibrary/Private/AbcFile.cpp b/Engine/Plugins/Experimental/AlembicImporter/Source/AlembicLibrary/Private/AbcFile.cpp
--- a/Engine/Plugins/Experimental/AlembicImporter/Source/AlembicLibrary/Private/AbcFile.cpp
+++ b/Engine/Plugins/Experimental/AlembicImporter/Source/AlembicLibrary/Private/AbcFile.cpp
@@ -168,7 +168,7 @@ EAbcImportError FAbcFile::Import(UAbcImportSettings* InImportSettings)
 	}
 
 	SecondsPerFrame = TimeStep;
-	ImportLength = (FMath::Max(StartFrameIndex, 0) + FrameSpan) * TimeStep;
+	ImportLength = FrameSpan * TimeStep;
 
 	// If the user opted to skip empty frames out of the specified frame range, calculate the time offset
 	if (SamplingSettings.bSkipEmpty)
diff --git a/Engine/Source/Runtime/Engine/Classes/Engine/StaticMeshActor.h b/Engine/Source/Runtime/Engine/Classes/Engine/StaticMeshActor.h
--- a/Engine/Source/Runtime/Engine/Classes/Engine/StaticMeshActor.h
+++ b/Engine/Source/Runtime/Engine/Classes/Engine/StaticMeshActor.h
@@ -37,6 +37,7 @@ public:
 	ENavDataGatheringMode NavigationGeometryGatheringMode;
 
 	/** Function to change mobility type */
+	UFUNCTION(BlueprintCallable, Category = Mobility)
 	void SetMobility(EComponentMobility::Type InMobility);
 
 	//~ Begin AActor Interface
