Deprecate FRawMesh, replace by MeshDescription - ABC Importer

diff --git a/Engine/Plugins/Experimental/AlembicImporter/Source/AlembicLibrary/AlembicLibrary.Build.cs b/Engine/Plugins/Experimental/AlembicImporter/Source/AlembicLibrary/AlembicLibrary.Build.cs
--- a/Engine/Plugins/Experimental/AlembicImporter/Source/AlembicLibrary/AlembicLibrary.Build.cs
+++ b/Engine/Plugins/Experimental/AlembicImporter/Source/AlembicLibrary/AlembicLibrary.Build.cs
@@ -13,7 +13,6 @@ public class AlembicLibrary : ModuleRules
 				"Engine",
                 "InputCore",
                 "UnrealEd",
-                "RawMesh",
                 "GeometryCache",
                 "AlembicLib",
                 "MeshUtilities",
@@ -33,7 +32,9 @@ public class AlembicLibrary : ModuleRules
                 "Core",
                 "CoreUObject",
                 "Engine",
-                "UnrealEd"
+                "UnrealEd",
+                "MeshDescription",
+                "MeshDescriptionOperations",
             }
         );
 
diff --git a/Engine/Plugins/Experimental/AlembicImporter/Source/AlembicLibrary/Private/AbcImporter.cpp b/Engine/Plugins/Experimental/AlembicImporter/Source/AlembicLibrary/Private/AbcImporter.cpp
--- a/Engine/Plugins/Experimental/AlembicImporter/Source/AlembicLibrary/Private/AbcImporter.cpp
+++ b/Engine/Plugins/Experimental/AlembicImporter/Source/AlembicLibrary/Private/AbcImporter.cpp
@@ -23,7 +23,10 @@ THIRD_PARTY_INCLUDES_END
 #include "Misc/ScopedSlowTask.h"
 
 #include "PackageTools.h"
-#include "RawMesh.h"
+#include "MeshDescription.h"
+#include "MeshAttributes.h"
+#include "MeshAttributeArray.h"
+#include "MeshDescriptionOperations.h"
 #include "ObjectTools.h"
 
 #include "Engine/StaticMesh.h"
@@ -179,7 +182,7 @@ T* FAbcImporter::CreateObjectInstance(UObject*& InParent, const FString& ObjectN
 	return NewObject<T>(Package, FName(*SanitizedObjectName), Flags | RF_Public);
 }
 
-UStaticMesh* FAbcImporter::CreateStaticMeshFromRawMesh(UObject* InParent, const FString& Name, EObjectFlags Flags, const uint32 NumMaterials, const TArray<FString>& FaceSetNames, FRawMesh& RawMesh)
+UStaticMesh* FAbcImporter::CreateStaticMeshFromSample(UObject* InParent, const FString& Name, EObjectFlags Flags, const uint32 NumMaterials, const TArray<FString>& FaceSetNames, const FAbcMeshSample* Sample)
 {
 	UStaticMesh* StaticMesh = CreateObjectInstance<UStaticMesh>(InParent, Name, Flags);
 
@@ -187,8 +190,10 @@ UStaticMesh* FAbcImporter::CreateStaticMeshFromRawMesh(UObject* InParent, const
 	if (StaticMesh)
 	{
 		// Add the first LOD, we only support one
+		int32 LODIndex = 0;
 		StaticMesh->AddSourceModel();
-
+		FMeshDescription* MeshDescription = StaticMesh->CreateOriginalMeshDescription(LODIndex);
+		StaticMesh->RegisterMeshAttributes(*MeshDescription);
 		// Generate a new lighting GUID (so its unique)
 		StaticMesh->LightingGuid = FGuid::NewGuid();
 
@@ -215,11 +220,13 @@ UStaticMesh* FAbcImporter::CreateStaticMeshFromRawMesh(UObject* InParent, const
 				Material = RetrieveMaterial(FaceSetNames[MaterialIndex], InParent, Flags);
 			}
 
-			StaticMesh->StaticMaterials.Add(( Material != nullptr) ? Material : DefaultMaterial);
+			StaticMesh->StaticMaterials.Add((Material != nullptr) ? Material : DefaultMaterial);
 		}
 
+		GenerateMeshDescriptionFromSample(Sample, MeshDescription, StaticMesh);
+
 		// Get the first LOD for filling it up with geometry, only support one LOD
-		FStaticMeshSourceModel& SrcModel = StaticMesh->SourceModels[0];
+		FStaticMeshSourceModel& SrcModel = StaticMesh->SourceModels[LODIndex];
 		// Set build settings for the static mesh
 		SrcModel.BuildSettings.bRecomputeNormals = false;
 		SrcModel.BuildSettings.bRecomputeTangents = false;
@@ -229,9 +236,9 @@ UStaticMesh* FAbcImporter::CreateStaticMeshFromRawMesh(UObject* InParent, const
 		// Set lightmap UV index to 1 since we currently only import one set of UVs from the Alembic Data file
 		SrcModel.BuildSettings.DstLightmapIndex = 1;
 
-		// Store the raw mesh within the RawMeshBulkData
-		SrcModel.SaveRawMesh(RawMesh);
-		
+		// Store the mesh description
+		StaticMesh->CommitOriginalMeshDescription(LODIndex);
+
 		//Set the Imported version before calling the build
 		StaticMesh->ImportVersion = EImportStaticMeshVersion::LastVersion;
 
@@ -256,71 +263,65 @@ const TArray<UStaticMesh*> FAbcImporter::ImportAsStaticMesh(UObject* InParent, E
 	{
 		const TArray<FAbcPolyMesh*>& PolyMeshes = AbcFile->GetPolyMeshes();
 		if (StaticMeshSettings.bMergeMeshes)
-	{
-		// If merging we merge all the raw mesh structures together and generate a static mesh asset from this
-		TArray<FString> MergedFaceSetNames;
-		TArray<FAbcMeshSample*> Samples;
-		uint32 TotalNumMaterials = 0;
+		{
+			// If merging we merge all the raw mesh structures together and generate a static mesh asset from this
+			TArray<FString> MergedFaceSetNames;
+			TArray<FAbcMeshSample*> Samples;
+			uint32 TotalNumMaterials = 0;
 
 			TArray<const FAbcMeshSample*> SamplesToMerge;
 			// Should merge all samples in the Alembic cache to one single static mesh
 			for (const FAbcPolyMesh* PolyMesh : PolyMeshes)
-		{
-				if (PolyMesh->bShouldImport)
 			{
+				if (PolyMesh->bShouldImport)
+				{
 					const FAbcMeshSample* Sample = PolyMesh->GetSample(FrameIndex);
 					SamplesToMerge.Add(Sample);
-				TotalNumMaterials += (Sample->NumMaterials != 0) ? Sample->NumMaterials : 1;
+					TotalNumMaterials += (Sample->NumMaterials != 0) ? Sample->NumMaterials : 1;
 
 					if (PolyMesh->FaceSetNames.Num() > 0)
-				{
+					{
 						MergedFaceSetNames.Append(PolyMesh->FaceSetNames);
+					}
+					else
+					{
+						// Default name
+						static const FString DefaultName("NoFaceSetName");
+						MergedFaceSetNames.Add(DefaultName);
+					}
 				}
-				else
-				{
-					// Default name
-					static const FString DefaultName("NoFaceSetName");
-					MergedFaceSetNames.Add(DefaultName);
-			}
-		}
 			}
 
-		// Only merged samples if there are any
+			// Only merged samples if there are any
 			if (SamplesToMerge.Num())
-		{
+			{
 				FAbcMeshSample* MergedSample = AbcImporterUtilities::MergeMeshSamples(SamplesToMerge);
 
-			FRawMesh RawMesh;
-			GenerateRawMeshFromSample(MergedSample, RawMesh);
 
-				UStaticMesh* StaticMesh = CreateStaticMeshFromRawMesh(InParent, InParent != GetTransientPackage() ? FPaths::GetBaseFilename(InParent->GetName()) : (FPaths::GetBaseFilename(AbcFile->GetFilePath()) + "_" + FGuid::NewGuid().ToString()), Flags, TotalNumMaterials, MergedFaceSetNames, RawMesh);
-			if (StaticMesh)
-			{
-					ImportedStaticMeshes.Add(StaticMesh);
+				UStaticMesh* StaticMesh = CreateStaticMeshFromSample(InParent, InParent != GetTransientPackage() ? FPaths::GetBaseFilename(InParent->GetName()) : (FPaths::GetBaseFilename(AbcFile->GetFilePath()) + "_" + FGuid::NewGuid().ToString()), Flags, TotalNumMaterials, MergedFaceSetNames, MergedSample);
+				if (StaticMesh)
+				{
+						ImportedStaticMeshes.Add(StaticMesh);
+				}
 			}
 		}
-	}
-	else
-	{
-			for (const FAbcPolyMesh* PolyMesh : PolyMeshes)
+		else
 		{
+			for (const FAbcPolyMesh* PolyMesh : PolyMeshes)
+			{
 				const FAbcMeshSample* Sample = PolyMesh->GetSample(FrameIndex);
 				if (PolyMesh->bShouldImport && Sample)
-			{
-					// Populate raw mesh from sample
-					FRawMesh RawMesh;
-					GenerateRawMeshFromSample(Sample, RawMesh);
-
+				{
 					// Setup static mesh instance
-					UStaticMesh* StaticMesh = CreateStaticMeshFromRawMesh(InParent, InParent != GetTransientPackage() ? PolyMesh->GetName() : PolyMesh->GetName() + "_" + FGuid::NewGuid().ToString(), Flags, Sample->NumMaterials, PolyMesh->FaceSetNames, RawMesh);
+					UStaticMesh* StaticMesh = CreateStaticMeshFromSample(InParent, InParent != GetTransientPackage() ? PolyMesh->GetName() : PolyMesh->GetName() + "_" + FGuid::NewGuid().ToString(), Flags, Sample->NumMaterials, PolyMesh->FaceSetNames, Sample);
 
-				if (StaticMesh)
-				{
+					if (StaticMesh)
+					{
 						ImportedStaticMeshes.Add(StaticMesh);
+					}
 				}
 			}
 		}
-	}
 	};
 	
 
@@ -1253,37 +1254,108 @@ const uint32 FAbcImporter::GetNumMeshTracks() const
 	return (AbcFile != nullptr) ? AbcFile->GetNumPolyMeshes() : 0;
 }
 
-void FAbcImporter::GenerateRawMeshFromSample( const FAbcMeshSample* Sample, FRawMesh& RawMesh)
+void FAbcImporter::GenerateMeshDescriptionFromSample(const FAbcMeshSample* Sample, FMeshDescription* MeshDescription, UStaticMesh* StaticMesh)
 {
-	// Set vertex data for mesh
-	RawMesh.VertexPositions = Sample->Vertices;
-
-	// Copy over per-index based data
-	RawMesh.WedgeIndices = Sample->Indices;
-	RawMesh.WedgeTangentX = Sample->TangentX;
-	RawMesh.WedgeTangentY = Sample->TangentY;
-	RawMesh.WedgeTangentZ = Sample->Normals;
+	TVertexAttributeArray<FVector>& VertexPositions = MeshDescription->VertexAttributes().GetAttributes<FVector>(MeshAttribute::Vertex::Position);
+	TEdgeAttributeArray<bool>& EdgeHardnesses = MeshDescription->EdgeAttributes().GetAttributes<bool>(MeshAttribute::Edge::IsHard);
+	TEdgeAttributeArray<float>& EdgeCreaseSharpnesses = MeshDescription->EdgeAttributes().GetAttributes<float>(MeshAttribute::Edge::CreaseSharpness);
+	TPolygonGroupAttributeArray<FName>& PolygonGroupImportedMaterialSlotNames = MeshDescription->PolygonGroupAttributes().GetAttributes<FName>(MeshAttribute::PolygonGroup::ImportedMaterialSlotName);
+	TVertexInstanceAttributeArray<FVector>& VertexInstanceNormals = MeshDescription->VertexInstanceAttributes().GetAttributes<FVector>(MeshAttribute::VertexInstance::Normal);
+	TVertexInstanceAttributeArray<FVector>& VertexInstanceTangents = MeshDescription->VertexInstanceAttributes().GetAttributes<FVector>(MeshAttribute::VertexInstance::Tangent);
+	TVertexInstanceAttributeArray<float>& VertexInstanceBinormalSigns = MeshDescription->VertexInstanceAttributes().GetAttributes<float>(MeshAttribute::VertexInstance::BinormalSign);
+	TVertexInstanceAttributeArray<FVector4>& VertexInstanceColors = MeshDescription->VertexInstanceAttributes().GetAttributes<FVector4>(MeshAttribute::VertexInstance::Color);
+	TVertexInstanceAttributeIndicesArray<FVector2D>& VertexInstanceUVs = MeshDescription->VertexInstanceAttributes().GetAttributesSet<FVector2D>(MeshAttribute::VertexInstance::TextureCoordinate);
+
+	//Speedtree use UVs to store is data
+	VertexInstanceUVs.SetNumIndices(Sample->NumUVSets);
+	
+	for (int32 MatIndex = 0; MatIndex < StaticMesh->StaticMaterials.Num(); ++MatIndex)
+	{
+		const FPolygonGroupID& PolygonGroupID = MeshDescription->CreatePolygonGroup();
+		PolygonGroupImportedMaterialSlotNames[PolygonGroupID] = StaticMesh->StaticMaterials[MatIndex].ImportedMaterialSlotName;
+	}
 
-	for (uint32 UVIndex = 0; UVIndex < Sample->NumUVSets; ++UVIndex)
+	// position
+	for (int32 VertexIndex = 0; VertexIndex < Sample->Vertices.Num(); ++VertexIndex)
 	{
-		RawMesh.WedgeTexCoords[UVIndex] = Sample->UVs[UVIndex];
+		FVector Position = Sample->Vertices[VertexIndex];
+
+		FVertexID VertexID = MeshDescription->CreateVertex();
+		VertexPositions[VertexID] = FVector(Position);
 	}
 
-	if ( Sample->Colors.Num() )
+	uint32 TriangleCount = Sample->Indices.Num() / 3;
+	for (uint32 TriangleIndex = 0; TriangleIndex < TriangleCount; ++TriangleIndex)
 	{
-		for (const FLinearColor& LinearColor : Sample->Colors)
+		FVertexInstanceID CornerVertexInstanceIDs[3];
+		FVertexID CornerVertexIDs[3];
+		for (int32 Corner = 0; Corner < 3; ++Corner)
+		{
+			uint32 IndiceIndex = (TriangleIndex * 3) + Corner;
+			uint32 VertexIndex = Sample->Indices[IndiceIndex];
+			const FVertexID VertexID(VertexIndex);
+			const FVertexInstanceID& VertexInstanceID = MeshDescription->CreateVertexInstance(VertexID);
+
+			// tangents
+			FVector TangentX = Sample->TangentX[IndiceIndex];
+			FVector TangentY = Sample->TangentY[IndiceIndex];
+			FVector TangentZ = Sample->Normals[IndiceIndex];
+
+			VertexInstanceTangents[VertexInstanceID] = TangentX;
+			VertexInstanceNormals[VertexInstanceID] = TangentZ;
+			VertexInstanceBinormalSigns[VertexInstanceID] = GetBasisDeterminantSign(TangentX.GetSafeNormal(), TangentY.GetSafeNormal(), TangentZ.GetSafeNormal());
+
+			if (Sample->Colors.Num())
+			{
+				VertexInstanceColors[VertexInstanceID] = FVector4(Sample->Colors[IndiceIndex]);
+			}
+			else
+			{
+				VertexInstanceColors[VertexInstanceID] = FVector4(FLinearColor::White);
+			}
+
+			for (uint32 UVIndex = 0; UVIndex < Sample->NumUVSets; ++UVIndex)
+			{
+				VertexInstanceUVs.GetArrayForIndex(UVIndex)[VertexInstanceID] = Sample->UVs[UVIndex][IndiceIndex];
+			}
+			CornerVertexInstanceIDs[Corner] = VertexInstanceID;
+			CornerVertexIDs[Corner] = VertexID;
+		}
+
+		TArray<FMeshDescription::FContourPoint> Contours;
+		for (int32 Corner = 0; Corner < 3; ++Corner)
 		{
-			RawMesh.WedgeColors.Add(LinearColor.ToFColor(false));
+			int32 ContourPointIndex = Contours.AddDefaulted();
+			FMeshDescription::FContourPoint& ContourPoint = Contours[ContourPointIndex];
+			//Find the matching edge ID
+			uint32 CornerIndices[2];
+			CornerIndices[0] = (Corner + 0) % 3;
+			CornerIndices[1] = (Corner + 1) % 3;
+
+			FVertexID EdgeVertexIDs[2];
+			EdgeVertexIDs[0] = CornerVertexIDs[CornerIndices[0]];
+			EdgeVertexIDs[1] = CornerVertexIDs[CornerIndices[1]];
+
+			FEdgeID MatchEdgeId = MeshDescription->GetVertexPairEdge(EdgeVertexIDs[0], EdgeVertexIDs[1]);
+			if (MatchEdgeId == FEdgeID::Invalid)
+			{
+				MatchEdgeId = MeshDescription->CreateEdge(EdgeVertexIDs[0], EdgeVertexIDs[1]);
+				//Conversion from smooth group to hard edges will be done at the end
+				EdgeHardnesses[MatchEdgeId] = false;
+				EdgeCreaseSharpnesses[MatchEdgeId] = 0.0f;
+			}
+			ContourPoint.EdgeID = MatchEdgeId;
+			ContourPoint.VertexInstanceID = CornerVertexInstanceIDs[CornerIndices[0]];
 		}
+		const FPolygonGroupID PolygonGroupID(Sample->MaterialIndices[TriangleIndex]);
+		// Insert a polygon into the mesh
+		const FPolygonID NewPolygonID = MeshDescription->CreatePolygon(PolygonGroupID, Contours);
+		//Triangulate the polygon
+		FMeshPolygon& Polygon = MeshDescription->GetPolygon(NewPolygonID);
+		MeshDescription->ComputePolygonTriangulation(NewPolygonID, Polygon.Triangles);
 	}
-	else
-	{
-		RawMesh.WedgeColors.AddDefaulted(RawMesh.WedgeIndices.Num());
-	}
-
-	// Copy over per-face data
-	RawMesh.FaceMaterialIndices = Sample->MaterialIndices;
-	RawMesh.FaceSmoothingMasks = Sample->SmoothingGroupIndices;
+	//Set the edge hardness from the smooth group
+	FMeshDescriptionOperations::ConvertSmoothGroupToHardEdges(Sample->SmoothingGroupIndices, *MeshDescription);
 }
 
 void FAbcImporter::GeometryCacheDataForMeshSample(FGeometryCacheMeshData &OutMeshData, const FAbcMeshSample* MeshSample, const uint32 MaterialOffset)
diff --git a/Engine/Plugins/Experimental/AlembicImporter/Source/AlembicLibrary/Public/AbcImporter.h b/Engine/Plugins/Experimental/AlembicImporter/Source/AlembicLibrary/Public/AbcImporter.h
--- a/Engine/Plugins/Experimental/AlembicImporter/Source/AlembicLibrary/Public/AbcImporter.h
+++ b/Engine/Plugins/Experimental/AlembicImporter/Source/AlembicLibrary/Public/AbcImporter.h
@@ -26,12 +26,12 @@ class UGeometryCacheTrackStreamable;
 class UAbcImportSettings;
 class FSkeletalMeshImportData;
 class UAbcAssetImportData;
+class FMeshDescription;
 
 struct FAbcImportData;
 struct FGeometryCacheMeshData;
 struct FAbcMeshSample;
 struct FAbcCompressionSettings;
-struct FRawMesh;
 
 enum EAbcImportError : uint32
 {
@@ -169,20 +169,20 @@ private:
 	void GeometryCacheDataForMeshSample(FGeometryCacheMeshData &OutMeshData, const FAbcMeshSample* MeshSample, const uint32 MaterialOffset);
 
 	/**
-	* Creates a Static mesh from the given raw mesh
-	*	
+	* Creates a Static mesh from the given mesh description
+	*
 	* @param InParent - ParentObject for the static mesh
 	* @param Name - Name for the static mesh
 	* @param Flags - Object flags
 	* @param NumMaterials - Number of materials to add
 	* @param FaceSetNames - Face set names used for retrieving the materials
-	* @param RawMesh - The RawMesh from which the static mesh should be constructed
+	* @param MeshDescription - The MeshDescription from which the static mesh should be constructed
 	* @return UStaticMesh*
 	*/
-	UStaticMesh* CreateStaticMeshFromRawMesh(UObject* InParent, const FString& Name, EObjectFlags Flags, const uint32 NumMaterials, const TArray<FString>& FaceSetNames, FRawMesh& RawMesh);
+	UStaticMesh* CreateStaticMeshFromSample(UObject* InParent, const FString& Name, EObjectFlags Flags, const uint32 NumMaterials, const TArray<FString>& FaceSetNames, const FAbcMeshSample* Sample);
 
-	/** Generates and populate a FRawMesh instance from the given sample*/
-	void GenerateRawMeshFromSample(const FAbcMeshSample* Sample, FRawMesh& RawMesh);
+	/** Generates and populate a FMeshDescription instance from the given sample*/
+	void GenerateMeshDescriptionFromSample(const FAbcMeshSample* Sample, FMeshDescription* MeshDescription, UStaticMesh* StaticMesh);
 	
 	/** Retrieves a material according to the given name and resaves it into the parent package*/
 	UMaterialInterface* RetrieveMaterial(const FString& MaterialName, UObject* InParent, EObjectFlags Flags );
diff --git a/Engine/Source/Developer/MeshDescriptionOperations/Private/MeshDescriptionOperations.cpp b/Engine/Source/Developer/MeshDescriptionOperations/Private/MeshDescriptionOperations.cpp
--- a/Engine/Source/Developer/MeshDescriptionOperations/Private/MeshDescriptionOperations.cpp
+++ b/Engine/Source/Developer/MeshDescriptionOperations/Private/MeshDescriptionOperations.cpp
@@ -185,7 +185,7 @@ void FMeshDescriptionOperations::ConvertHardEdgesToSmoothGroup(const FMeshDescri
 	}
 }
 
-void FMeshDescriptionOperations::ConvertSmoothGroupToHardEdges(const FRawMesh& SourceRawMesh, FMeshDescription& DestinationMeshDescription)
+void FMeshDescriptionOperations::ConvertSmoothGroupToHardEdges(const TArray<uint32>& FaceSmoothingMasks, FMeshDescription& DestinationMeshDescription)
 {
 	TEdgeAttributeArray<bool>& EdgeHardnesses = DestinationMeshDescription.EdgeAttributes().GetAttributes<bool>(MeshAttribute::Edge::IsHard);
 
@@ -203,8 +203,8 @@ void FMeshDescriptionOperations::ConvertSmoothGroupToHardEdges(const FRawMesh& S
 		{
 			FPolygonID CurrentPolygonID = ConnectedPolygons.Pop(true);
 			int32 CurrentPolygonIDValue = CurrentPolygonID.GetValue();
-			check(SourceRawMesh.FaceSmoothingMasks.IsValidIndex(CurrentPolygonIDValue));
-			const uint32 ReferenceSmoothGroup = SourceRawMesh.FaceSmoothingMasks[CurrentPolygonIDValue];
+			check(FaceSmoothingMasks.IsValidIndex(CurrentPolygonIDValue));
+			const uint32 ReferenceSmoothGroup = FaceSmoothingMasks[CurrentPolygonIDValue];
 			TArray<FEdgeID> PolygonEdges;
 			DestinationMeshDescription.GetPolygonEdges(CurrentPolygonID, PolygonEdges);
 			for (const FEdgeID& EdgeID : PolygonEdges)
@@ -222,8 +222,8 @@ void FMeshDescriptionOperations::ConvertSmoothGroupToHardEdges(const FRawMesh& S
 					{
 						continue;
 					}
-					check(SourceRawMesh.FaceSmoothingMasks.IsValidIndex(EdgePolygonIDValue));
-					const uint32 TestSmoothGroup = SourceRawMesh.FaceSmoothingMasks[EdgePolygonIDValue];
+					check(FaceSmoothingMasks.IsValidIndex(EdgePolygonIDValue));
+					const uint32 TestSmoothGroup = FaceSmoothingMasks[EdgePolygonIDValue];
 					if ((TestSmoothGroup & ReferenceSmoothGroup) == 0)
 					{
 						EdgeHardnesses[EdgeID] = true;
@@ -571,7 +571,7 @@ void FMeshDescriptionOperations::ConvertFromRawMesh(const FRawMesh& SourceRawMes
 		}
 	}
 	
-	ConvertSmoothGroupToHardEdges(SourceRawMesh, DestinationMeshDescription);
+	ConvertSmoothGroupToHardEdges(SourceRawMesh.FaceSmoothingMasks, DestinationMeshDescription);
 
 	//Create the missing normals and tangents, should we use Mikkt space for tangent???
 	if (!bHasNormals || !bHasTangents)
diff --git a/Engine/Source/Developer/MeshDescriptionOperations/Public/MeshDescriptionOperations.h b/Engine/Source/Developer/MeshDescriptionOperations/Public/MeshDescriptionOperations.h
--- a/Engine/Source/Developer/MeshDescriptionOperations/Public/MeshDescriptionOperations.h
+++ b/Engine/Source/Developer/MeshDescriptionOperations/Public/MeshDescriptionOperations.h
@@ -64,9 +64,8 @@ public:
 	/** Create some UVs from the specified mesh description data. */
 	static bool GenerateUniqueUVsForStaticMesh(const FMeshDescription& MeshDescription, int32 TextureResolution, TArray<FVector2D>& OutTexCoords);
 
-private:
 	
 	static void ConvertHardEdgesToSmoothGroup(const FMeshDescription& SourceMeshDescription, FRawMesh& DestinationRawMesh);
 
-	static void ConvertSmoothGroupToHardEdges(const FRawMesh& SourceRawMesh, FMeshDescription& DestinationMeshDescription);
+	static void ConvertSmoothGroupToHardEdges(const TArray<uint32>& FaceSmoothingMasks, FMeshDescription& DestinationMeshDescription);
 };
\ No newline at end of file
