Undefine some additional Windows macros after including WindowsHWrapper.h. Also add Windows::TRUE and Windows::FALSE as namespaced constants which are equivalent to Windows' TRUE and FALSE macros.

diff --git a/Engine/Plugins/Experimental/AlembicImporter/Source/AlembicLibrary/Private/AbcFile.cpp b/Engine/Plugins/Experimental/AlembicImporter/Source/AlembicLibrary/Private/AbcFile.cpp
--- a/Engine/Plugins/Experimental/AlembicImporter/Source/AlembicLibrary/Private/AbcFile.cpp
+++ b/Engine/Plugins/Experimental/AlembicImporter/Source/AlembicLibrary/Private/AbcFile.cpp
@@ -16,11 +16,20 @@
 
 #include "HAL/Platform.h"
 
+
+#if PLATFORM_WINDOWS
+#include "Windows/AllowWindowsPlatformTypes.h"
+#endif
+
 THIRD_PARTY_INCLUDES_START
 #include <Alembic/AbcGeom/All.h>
 #include "Materials/MaterialInstance.h"
 THIRD_PARTY_INCLUDES_END
 
+#if PLATFORM_WINDOWS
+#include "Windows/HideWindowsPlatformTypes.h"
+#endif
+
 #define LOCTEXT_NAMESPACE "AbcFile"
 
 FAbcFile::FAbcFile(const FString& InFilePath)
diff --git a/Engine/Plugins/Experimental/AlembicImporter/Source/AlembicLibrary/Private/AbcImportUtilities.cpp b/Engine/Plugins/Experimental/AlembicImporter/Source/AlembicLibrary/Private/AbcImportUtilities.cpp
--- a/Engine/Plugins/Experimental/AlembicImporter/Source/AlembicLibrary/Private/AbcImportUtilities.cpp
+++ b/Engine/Plugins/Experimental/AlembicImporter/Source/AlembicLibrary/Private/AbcImportUtilities.cpp
@@ -10,6 +10,10 @@
 #include "Windows/WindowsHWrapper.h"
 #endif
 
+#if PLATFORM_WINDOWS
+#include "Windows/AllowWindowsPlatformTypes.h"
+#endif
+
 THIRD_PARTY_INCLUDES_START
 #include <Alembic/AbcCoreHDF5/All.h>
 #include <Alembic/AbcCoreOgawa/All.h>
@@ -21,6 +25,10 @@ THIRD_PARTY_INCLUDES_START
 #include <Alembic/Abc/IObject.h>
 THIRD_PARTY_INCLUDES_END
 
+#if PLATFORM_WINDOWS
+#include "Windows/HideWindowsPlatformTypes.h"
+#endif
+
 #include "Rendering/SkeletalMeshLODModel.h"
 
 #define LOCTEXT_NAMESPACE "AbcImporterUtilities"
diff --git a/Engine/Plugins/Experimental/AlembicImporter/Source/AlembicLibrary/Private/AbcImportUtilities.h b/Engine/Plugins/Experimental/AlembicImporter/Source/AlembicLibrary/Private/AbcImportUtilities.h
--- a/Engine/Plugins/Experimental/AlembicImporter/Source/AlembicLibrary/Private/AbcImportUtilities.h
+++ b/Engine/Plugins/Experimental/AlembicImporter/Source/AlembicLibrary/Private/AbcImportUtilities.h
@@ -6,6 +6,7 @@
 
 #if PLATFORM_WINDOWS
 #include "Windows/WindowsHWrapper.h"
+#include "Windows/AllowWindowsPlatformTypes.h"
 #endif
 
 THIRD_PARTY_INCLUDES_START
@@ -14,6 +15,10 @@ THIRD_PARTY_INCLUDES_START
 #include <Alembic/AbcGeom/All.h>
 THIRD_PARTY_INCLUDES_END
 
+#if PLATFORM_WINDOWS
+#include "Windows/HideWindowsPlatformTypes.h"
+#endif
+
 #include "GeometryCache.h"
 #include "GeometryCacheTrackFlipbookAnimation.h"
 #include "GeometryCacheTrackTransformAnimation.h"
diff --git a/Engine/Plugins/Experimental/AlembicImporter/Source/AlembicLibrary/Public/AbcFile.h b/Engine/Plugins/Experimental/AlembicImporter/Source/AlembicLibrary/Public/AbcFile.h
--- a/Engine/Plugins/Experimental/AlembicImporter/Source/AlembicLibrary/Public/AbcFile.h
+++ b/Engine/Plugins/Experimental/AlembicImporter/Source/AlembicLibrary/Public/AbcFile.h
@@ -7,6 +7,7 @@
 
 #if PLATFORM_WINDOWS
 #include "Windows/WindowsHWrapper.h"
+#include "Windows/AllowWindowsPlatformTypes.h"
 #endif
 
 THIRD_PARTY_INCLUDES_START
@@ -16,6 +17,10 @@ THIRD_PARTY_INCLUDES_START
 
 THIRD_PARTY_INCLUDES_END
 
+#if PLATFORM_WINDOWS
+#include "Windows/HideWindowsPlatformTypes.h"
+#endif
+
 class IAbcObject;
 class FAbcPolyMesh;
 class FAbcTransform;
diff --git a/Engine/Plugins/Experimental/AlembicImporter/Source/AlembicLibrary/Public/AbcImporter.h b/Engine/Plugins/Experimental/AlembicImporter/Source/AlembicLibrary/Public/AbcImporter.h
--- a/Engine/Plugins/Experimental/AlembicImporter/Source/AlembicLibrary/Public/AbcImporter.h
+++ b/Engine/Plugins/Experimental/AlembicImporter/Source/AlembicLibrary/Public/AbcImporter.h
@@ -10,12 +10,17 @@
 
 #if PLATFORM_WINDOWS
 #include "Windows/WindowsHWrapper.h"
+#include "Windows/AllowWindowsPlatformTypes.h"
 #endif
 
 THIRD_PARTY_INCLUDES_START
 #include <Alembic/AbcGeom/All.h>
 THIRD_PARTY_INCLUDES_END
 
+#if PLATFORM_WINDOWS
+#include "Windows/HideWindowsPlatformTypes.h"
+#endif
+
 class UMaterial;
 class UStaticMesh;
 class USkeletalMesh;
diff --git a/Engine/Plugins/Experimental/AlembicImporter/Source/AlembicLibrary/Public/AbcObject.h b/Engine/Plugins/Experimental/AlembicImporter/Source/AlembicLibrary/Public/AbcObject.h
--- a/Engine/Plugins/Experimental/AlembicImporter/Source/AlembicLibrary/Public/AbcObject.h
+++ b/Engine/Plugins/Experimental/AlembicImporter/Source/AlembicLibrary/Public/AbcObject.h
@@ -7,6 +7,7 @@
 
 #if PLATFORM_WINDOWS
 #include "Windows/WindowsHWrapper.h"
+#include "Windows/AllowWindowsPlatformTypes.h"
 #endif
 
 THIRD_PARTY_INCLUDES_START
@@ -14,6 +15,10 @@ THIRD_PARTY_INCLUDES_START
 #include <Alembic/AbcGeom/All.h>
 THIRD_PARTY_INCLUDES_END
 
+#if PLATFORM_WINDOWS
+#include "Windows/HideWindowsPlatformTypes.h"
+#endif
+
 class FAbcFile;
 enum class EFrameReadFlags : uint8;
 
diff --git a/Engine/Source/Runtime/Core/Public/Windows/MinimalWindowsApi.h b/Engine/Source/Runtime/Core/Public/Windows/MinimalWindowsApi.h
--- a/Engine/Source/Runtime/Core/Public/Windows/MinimalWindowsApi.h
+++ b/Engine/Source/Runtime/Core/Public/Windows/MinimalWindowsApi.h
@@ -85,6 +85,10 @@ namespace Windows
 
 	typedef RTL_SRWLOCK SRWLOCK, *PSRWLOCK;
 
+	// Constants
+	static const BOOL TRUE = 1;
+	static const BOOL FALSE = 0;
+
 	// Modules
 	MINIMAL_WINDOWS_API HMODULE WINAPI LoadLibraryW(LPCTSTR lpFileName);
 	MINIMAL_WINDOWS_API BOOL WINAPI FreeLibrary(HMODULE hModule);
diff --git a/Engine/Source/Runtime/Core/Public/Windows/PostWindowsApi.h b/Engine/Source/Runtime/Core/Public/Windows/PostWindowsApi.h
--- a/Engine/Source/Runtime/Core/Public/Windows/PostWindowsApi.h
+++ b/Engine/Source/Runtime/Core/Public/Windows/PostWindowsApi.h
@@ -3,6 +3,14 @@
 // Re-enable warnings
 THIRD_PARTY_INCLUDES_END
 
+// Hide Windows-only types (same as HideWindowsPlatformTypes.h)
+#undef INT
+#undef UINT
+#undef DWORD
+#undef FLOAT
+#undef TRUE
+#undef FALSE
+
 // Undo any Windows defines.
 #undef uint8
 #undef uint16
@@ -36,6 +44,7 @@ THIRD_PARTY_INCLUDES_END
 #undef ReportEvent
 #undef GetClassName
 #undef GetClassInfo
+#undef Yield
 #undef IMediaEventSink
 
 // Undefine all the atomics. AllowWindowsPlatformAtomics/HideWindowsPlatformAtomics temporarily defining these macros.
diff --git a/Engine/Source/Runtime/Slate/Private/Framework/Application/HardwareCursor.cpp b/Engine/Source/Runtime/Slate/Private/Framework/Application/HardwareCursor.cpp
--- a/Engine/Source/Runtime/Slate/Private/Framework/Application/HardwareCursor.cpp
+++ b/Engine/Source/Runtime/Slate/Private/Framework/Application/HardwareCursor.cpp
@@ -125,7 +125,7 @@ void FHardwareCursor::CreateCursorFromRGBABuffer(const FColor* Pixels, int32 Wid
 	HBITMAP CursorMask = ::CreateBitmap(Width, Height, 1, 8, MaskPixels.GetData());
 
 	ICONINFO IconInfo = { 0 };
-	IconInfo.fIcon = FALSE;
+	IconInfo.fIcon = Windows::FALSE;
 	IconInfo.xHotspot = FMath::RoundToInt(InHotSpot.X * Width);
 	IconInfo.yHotspot = FMath::RoundToInt(InHotSpot.Y * Height);
 	IconInfo.hbmColor = CursorColor;
