Alembic Problem - Missing keys in animations when selecting time range

diff --git a/Engine/Plugins/Experimental/AlembicImporter/Source/AlembicLibrary/Private/AbcFile.cpp b/Engine/Plugins/Experimental/AlembicImporter/Source/AlembicLibrary/Private/AbcFile.cpp
--- a/Engine/Plugins/Experimental/AlembicImporter/Source/AlembicLibrary/Private/AbcFile.cpp
+++ b/Engine/Plugins/Experimental/AlembicImporter/Source/AlembicLibrary/Private/AbcFile.cpp
@@ -170,11 +170,8 @@ EAbcImportError FAbcFile::Import(UAbcImportSettings* InImportSettings)
 	SecondsPerFrame = TimeStep;
 	ImportLength = FrameSpan * TimeStep;
 
-	// If the user opted to skip empty frames out of the specified frame range, calculate the time offset
-	if (SamplingSettings.bSkipEmpty)
-	{
-		ImportTimeOffset = StartFrameIndex * SecondsPerFrame;
-	}
+	// Calculate time offset from start of import animation range
+	ImportTimeOffset = StartFrameIndex * SecondsPerFrame;
 
 	// Read first-frames for both the transforms and poly meshes
 
