Regression: Importing an Alembic File with Geometry cache or Skeletal leads to crash.
#fix Ensure we check the Constant Topology flag, this was lost during the previous refactor
#misc Updated warning/error import message to notify user about non-constant topology meshes

diff --git a/Engine/Plugins/Experimental/AlembicImporter/Source/AlembicLibrary/Private/AbcImporter.cpp b/Engine/Plugins/Experimental/AlembicImporter/Source/AlembicLibrary/Private/AbcImporter.cpp
--- a/Engine/Plugins/Experimental/AlembicImporter/Source/AlembicLibrary/Private/AbcImporter.cpp
+++ b/Engine/Plugins/Experimental/AlembicImporter/Source/AlembicLibrary/Private/AbcImporter.cpp
@@ -838,7 +838,7 @@ const bool FAbcImporter::CompressAnimationDataUsingPCA(const FAbcCompressionSett
 	TArray<FAbcPolyMesh*> ConstantPolyMeshObjects;
 	for (FAbcPolyMesh* PolyMesh : PolyMeshes)
 	{
-		if (PolyMesh->bShouldImport)
+		if (PolyMesh->bShouldImport && PolyMesh->bConstantTopology)
 		{
 			if (PolyMesh->IsConstant() && PolyMesh->bConstantTransformation)
 			{
@@ -1089,7 +1089,7 @@ const bool FAbcImporter::CompressAnimationDataUsingPCA(const FAbcCompressionSett
 	else
 	{
 		bResult = ConstantPolyMeshObjects.Num() > 0;
-		TSharedRef<FTokenizedMessage> Message = FTokenizedMessage::Create(bResult ? EMessageSeverity::Warning : EMessageSeverity::Error, LOCTEXT("NoMeshesToProcess", "Unable to compress animation data, no meshes found with Vertex Animation and baked Matrix Animation is turned off."));
+		TSharedRef<FTokenizedMessage> Message = FTokenizedMessage::Create(bResult ? EMessageSeverity::Warning : EMessageSeverity::Error, LOCTEXT("NoMeshesToProcess", "Unable to compress animation data, no meshes (with constant topology) found with Vertex Animation and baked Matrix Animation is turned off."));
 		FAbcImportLogger::AddImportMessage(Message);
 		
 	}
