Importing an Alembic File with "Skip Empty Frames at Start of Alembic Sequence" Enabled Causes Crash
#fix Ensure that if no normals are available in the cache file we correctly calculate them for each individual frame
    - Pulled out normal calculation process for a frame into shared function to align all uses

diff --git a/Engine/Plugins/Experimental/AlembicImporter/Source/AlembicLibrary/Private/AbcPolyMesh.cpp b/Engine/Plugins/Experimental/AlembicImporter/Source/AlembicLibrary/Private/AbcPolyMesh.cpp
--- a/Engine/Plugins/Experimental/AlembicImporter/Source/AlembicLibrary/Private/AbcPolyMesh.cpp
+++ b/Engine/Plugins/Experimental/AlembicImporter/Source/AlembicLibrary/Private/AbcPolyMesh.cpp
@@ -19,7 +19,7 @@ FAbcPolyMesh::FAbcPolyMesh(const Alembic::AbcGeom::IPolyMesh& InPolyMesh, const
 	// Retrieve schema and frame information		
 	NumSamples = Schema.getNumSamples();
 	bConstant = Schema.isConstant();
-	bConstantTopology = (Schema.getTopologyVariance() == Alembic::AbcGeom::kHomogeneousTopology) || bConstant;
+	bConstantTopology = (Schema.getTopologyVariance() != Alembic::AbcGeom::kHeterogeneousTopology) || bConstant;
 	bConstantVisibility = AbcImporterUtilities::IsObjectVisibilityConstant(InPolyMesh);
 	SelfBounds = AbcImporterUtilities::ExtractBounds(Schema.getSelfBoundsProperty());
 	ChildBounds = AbcImporterUtilities::ExtractBounds(Schema.getChildBoundsProperty());
@@ -53,48 +53,13 @@ bool FAbcPolyMesh::ReadFirstFrame(const float InTime, const int32 FrameIndex)
 	
 	if (FirstSample)
 	{
-		/*
-		Normal cases
-		* No normals
-		- OneSmoothing group -> smooth normals and zeroed out smoothing groups
-		- Compute smooth normals
-		* Normals
-		- OneSmoothing group -> smooth normals and zeroed out smoothing groups
-		- Recompute normals -> Compute normals, compute smoothing groups -> compute smooth normals
-		- else compute smoothing groups
-		*/
-
-		const bool bRecomputeNormals = File->GetImportSettings()->NormalGenerationSettings.bRecomputeNormals;
-		if (File->GetImportSettings()->NormalGenerationSettings.bForceOneSmoothingGroupPerObject && bRecomputeNormals)
+		if (FirstSample->Normals.Num() == 0)
 		{
-			AbcImporterUtilities::CalculateSmoothNormals(FirstSample);
-			FirstSample->SmoothingGroupIndices.AddZeroed(FirstSample->Indices.Num() / 3);
-			FirstSample->NumSmoothingGroups = 1;
-		}
-		else
-		{
-			const bool bNormalsAvailable = FirstSample->Normals.Num() != 0;
-
-			// Recompute the (hard) normals if the user opted to do so
-			if (bRecomputeNormals)
-			{
-				AbcImporterUtilities::CalculateNormals(FirstSample);
-			}
-			// Otherwise we'd expect normals to be available, if not we should assume the object has smooth normals (so calculate them)
-			else if (!bNormalsAvailable)
-			{
-				AbcImporterUtilities::CalculateSmoothNormals(FirstSample);
-			}
+			// Normals are not available so they should be calculated each frame, so force set the Normals read flag which will ensure they are generated
+			SampleReadFlags |= ESampleReadFlags::Normals;
+		}	
 
-			// Generate smoothing groups from the normals to use for follow samples
-			AbcImporterUtilities::GenerateSmoothingGroupsIndices(FirstSample, File->GetImportSettings()->NormalGenerationSettings.HardEdgeAngleThreshold);
-
-			// In case we are expected to recompute the normals now recalculate the normals using the calculated smoothing groups
-			if (bRecomputeNormals)
-			{
-				AbcImporterUtilities::CalculateNormalsWithSmoothingGroups(FirstSample, FirstSample->SmoothingGroupIndices, FirstSample->NumSmoothingGroups);
-			}
-		}
+		CalculateNormalsForSample(FirstSample);
 
 		// Compute tangents for mesh
 		AbcImporterUtilities::ComputeTangents(FirstSample, File->GetImportSettings()->NormalGenerationSettings.bIgnoreDegenerateTriangles, *File->GetMeshUtilities());
@@ -127,6 +92,58 @@ bool FAbcPolyMesh::ReadFirstFrame(const float InTime, const int32 FrameIndex)
 	return FirstSample != nullptr;
 }
 
+void FAbcPolyMesh::CalculateNormalsForSample(FAbcMeshSample* Sample)
+{
+	/*
+	Normal cases
+	* No normals
+	- OneSmoothing group -> smooth normals and zeroed out smoothing groups
+	- Compute smooth normals
+	* Normals
+	- OneSmoothing group -> smooth normals and zeroed out smoothing groups
+	- Recompute normals -> Compute normals, compute smoothing groups -> compute smooth normals
+	- else compute smoothing groups
+	*/
+
+	const bool bRecomputeNormals = File->GetImportSettings()->NormalGenerationSettings.bRecomputeNormals;
+	if (File->GetImportSettings()->NormalGenerationSettings.bForceOneSmoothingGroupPerObject && bRecomputeNormals)
+	{
+		AbcImporterUtilities::CalculateSmoothNormals(Sample);
+		Sample->SmoothingGroupIndices.AddZeroed(Sample->Indices.Num() / 3);
+		Sample->NumSmoothingGroups = 1;
+	}
+	else
+	{
+		const bool bNormalsAvailable = Sample->Normals.Num() != 0;
+
+		// Recompute the (hard) normals if the user opted to do so
+		if (bRecomputeNormals)
+		{
+			AbcImporterUtilities::CalculateNormals(Sample);
+		}
+		// Otherwise we'd expect normals to be available, if not we should assume the object has smooth normals (so calculate them)
+		else if (!bNormalsAvailable)
+		{			
+			AbcImporterUtilities::CalculateSmoothNormals(Sample);
+			Sample->SmoothingGroupIndices.AddZeroed(Sample->Indices.Num() / 3);
+			Sample->NumSmoothingGroups = 1;
+		}
+
+		// If not smooth normals generate smoothing groups
+		if (bNormalsAvailable || bRecomputeNormals)
+		{
+			// Generate smoothing groups from the normals to use for following samples
+			AbcImporterUtilities::GenerateSmoothingGroupsIndices(Sample, File->GetImportSettings()->NormalGenerationSettings.HardEdgeAngleThreshold);
+		}
+
+		// In case we are expected to recompute the normals now recalculate the normals using the calculated smoothing groups
+		if (bRecomputeNormals)
+		{
+			AbcImporterUtilities::CalculateNormalsWithSmoothingGroups(Sample, Sample->SmoothingGroupIndices, Sample->NumSmoothingGroups);
+		}
+	}
+}
+
 void FAbcPolyMesh::SetFrameAndTime(const float InTime, const int32 FrameIndex, const EFrameReadFlags InFlags, const int32 TargetIndex /*= INDEX_NONE*/)
 {
 	if (!bShouldImport)
@@ -152,10 +169,20 @@ void FAbcPolyMesh::SetFrameAndTime(const float InTime, const int32 FrameIndex, c
 			const bool bVertexDataOnly = EnumHasAnyFlags(InFlags, EFrameReadFlags::PositionOnly);
 			WriteSample->Copy(FirstSample, bVertexDataOnly ? ESampleReadFlags::Positions : SampleReadFlags);
 			const bool bValidSample = AbcImporterUtilities::GenerateAbcMeshSampleDataForFrame(Schema, SampleSelector, WriteSample, bVertexDataOnly ? ESampleReadFlags::Positions : SampleReadFlags, InTime == MinTime);
+			// Check whether or not the number of normal indices matches with the first frame
+			const bool bMatchingIndices = FirstSample != nullptr && FirstSample->Indices.Num() == WriteSample->Indices.Num();
 			// Make sure in case of recomputing normals we enforece using the first sample data (otherwise we'll be using loaded or incorrectly calculated normals)
-			if (FirstSample != nullptr && (WriteSample->Normals.Num() == 0 || bRecomputeNormals))
+			if (WriteSample->Normals.Num() == 0 || bRecomputeNormals)
 			{
-				AbcImporterUtilities::CalculateNormalsWithSampleData(WriteSample, FirstSample);
+				// If the indices match we can recalculate the normals according to the first frame (and copy the smoothing indices)
+				if (bMatchingIndices)
+				{
+					AbcImporterUtilities::CalculateNormalsWithSampleData(WriteSample, FirstSample);
+				}
+				else
+				{
+					CalculateNormalsForSample(WriteSample);
+				}
 			}
 			else
 			{
diff --git a/Engine/Plugins/Experimental/AlembicImporter/Source/AlembicLibrary/Public/AbcPolyMesh.h b/Engine/Plugins/Experimental/AlembicImporter/Source/AlembicLibrary/Public/AbcPolyMesh.h
--- a/Engine/Plugins/Experimental/AlembicImporter/Source/AlembicLibrary/Public/AbcPolyMesh.h
+++ b/Engine/Plugins/Experimental/AlembicImporter/Source/AlembicLibrary/Public/AbcPolyMesh.h
@@ -111,6 +111,9 @@ public:
 	/** Whether or not this Mesh object should be imported */
 	bool bShouldImport;
 protected:
+	/** Calculate normals for sample according to available data and user settings */
+	void CalculateNormalsForSample(FAbcMeshSample* Sample);
+
 	/** Alembic representation of this object */
 	const Alembic::AbcGeom::IPolyMesh PolyMesh;
 	/** Schema extracted from Poly Mesh object  */
