Crash Occurs When Importing a .ABC File as Skeletal Mesh Into Project That Already Contains Asset's Materials
#fix make sure we call PostEditChanged for material assigned to new objects in Alembic importer, this flushes shader compilation and ensures we kick off a fresh one with necessary permutations

diff --git a/Engine/Plugins/Experimental/AlembicImporter/Source/AlembicLibrary/Private/AbcImporter.cpp b/Engine/Plugins/Experimental/AlembicImporter/Source/AlembicLibrary/Private/AbcImporter.cpp
--- a/Engine/Plugins/Experimental/AlembicImporter/Source/AlembicLibrary/Private/AbcImporter.cpp
+++ b/Engine/Plugins/Experimental/AlembicImporter/Source/AlembicLibrary/Private/AbcImporter.cpp
@@ -218,6 +218,7 @@ UStaticMesh* FAbcImporter::CreateStaticMeshFromSample(UObject* InParent, const F
 			if (FaceSetNames.IsValidIndex(MaterialIndex))
 			{
 				Material = RetrieveMaterial(FaceSetNames[MaterialIndex], InParent, Flags);
+				Material->PostEditChange();
 			}
 
 			StaticMesh->StaticMaterials.Add((Material != nullptr) ? Material : DefaultMaterial);
@@ -475,6 +476,7 @@ UGeometryCache* FAbcImporter::ImportAsGeometryCache(UObject* InParent, EObjectFl
 				{
 					UMaterialInterface* Material = RetrieveMaterial(FaceSetName, InParent, Flags);
 					GeometryCache->Materials.Add((Material != nullptr) ? Material : DefaultMaterial);		
+					Material->PostEditChange();
 				}
 			}
 			else
@@ -502,6 +504,7 @@ UGeometryCache* FAbcImporter::ImportAsGeometryCache(UObject* InParent, EObjectFl
 							if (PolyMesh->FaceSetNames.IsValidIndex(MaterialIndex))
 							{
 								Material = RetrieveMaterial(PolyMesh->FaceSetNames[MaterialIndex], InParent, Flags);
+								Material->PostEditChange();
 							}
 
 							GeometryCache->Materials.Add((Material != nullptr) ? Material : DefaultMaterial);
@@ -755,6 +758,7 @@ TArray<UObject*> FAbcImporter::ImportAsSkeletalMesh(UObject* InParent, EObjectFl
 				const FString& MaterialName = CompressedData.MaterialNames[MaterialIndex];
 				UMaterialInterface* Material = RetrieveMaterial(MaterialName, InParent, Flags);
 				SkeletalMesh->Materials.Add(FSkeletalMaterial(Material, true));
+				Material->PostEditChange();
 			}
 
 			++ObjectIndex;
