Fix for static analysis warning in Eigen library code

diff --git a/Engine/Plugins/Experimental/AlembicImporter/Source/AlembicLibrary/Private/EigenHelper.h b/Engine/Plugins/Experimental/AlembicImporter/Source/AlembicLibrary/Private/EigenHelper.h
--- a/Engine/Plugins/Experimental/AlembicImporter/Source/AlembicLibrary/Private/EigenHelper.h
+++ b/Engine/Plugins/Experimental/AlembicImporter/Source/AlembicLibrary/Private/EigenHelper.h
@@ -11,7 +11,7 @@
 
 #if defined(_MSC_VER) && USING_CODE_ANALYSIS
 	#pragma warning(push)
-	#pragma warning(disable:6326) // Potential comparison of a constant with another constant.
+	#pragma warning(disable:6294) // Ill-defined for-loop:  initial condition does not satisfy test.  Loop body not executed.
 #endif
 THIRD_PARTY_INCLUDES_START
 	#include <Eigen/Dense>
