Do not post-edit change DefaultMaterial when doing Alembic imports, this should fix time-outs for Alembic tests

diff --git a/Engine/Plugins/Experimental/AlembicImporter/Source/AlembicLibrary/Private/AbcImporter.cpp b/Engine/Plugins/Experimental/AlembicImporter/Source/AlembicLibrary/Private/AbcImporter.cpp
--- a/Engine/Plugins/Experimental/AlembicImporter/Source/AlembicLibrary/Private/AbcImporter.cpp
+++ b/Engine/Plugins/Experimental/AlembicImporter/Source/AlembicLibrary/Private/AbcImporter.cpp
@@ -218,7 +218,11 @@ UStaticMesh* FAbcImporter::CreateStaticMeshFromSample(UObject* InParent, const F
 			if (FaceSetNames.IsValidIndex(MaterialIndex))
 			{
 				Material = RetrieveMaterial(FaceSetNames[MaterialIndex], InParent, Flags);
-				Material->PostEditChange();
+
+				if (Material != UMaterial::GetDefaultMaterial(MD_Surface))
+				{
+					Material->PostEditChange();
+				}
 			}
 
 			StaticMesh->StaticMaterials.Add((Material != nullptr) ? Material : DefaultMaterial);
@@ -476,7 +480,11 @@ UGeometryCache* FAbcImporter::ImportAsGeometryCache(UObject* InParent, EObjectFl
 				{
 					UMaterialInterface* Material = RetrieveMaterial(FaceSetName, InParent, Flags);
 					GeometryCache->Materials.Add((Material != nullptr) ? Material : DefaultMaterial);		
-					Material->PostEditChange();
+
+					if (Material != UMaterial::GetDefaultMaterial(MD_Surface))
+					{
+						Material->PostEditChange();
+					}
 				}
 			}
 			else
@@ -504,7 +512,10 @@ UGeometryCache* FAbcImporter::ImportAsGeometryCache(UObject* InParent, EObjectFl
 							if (PolyMesh->FaceSetNames.IsValidIndex(MaterialIndex))
 							{
 								Material = RetrieveMaterial(PolyMesh->FaceSetNames[MaterialIndex], InParent, Flags);
-								Material->PostEditChange();
+								if (Material != UMaterial::GetDefaultMaterial(MD_Surface))
+								{
+									Material->PostEditChange();
+								}
 							}
 
 							GeometryCache->Materials.Add((Material != nullptr) ? Material : DefaultMaterial);
@@ -758,7 +769,10 @@ TArray<UObject*> FAbcImporter::ImportAsSkeletalMesh(UObject* InParent, EObjectFl
 				const FString& MaterialName = CompressedData.MaterialNames[MaterialIndex];
 				UMaterialInterface* Material = RetrieveMaterial(MaterialName, InParent, Flags);
 				SkeletalMesh->Materials.Add(FSkeletalMaterial(Material, true));
-				Material->PostEditChange();
+				if (Material != UMaterial::GetDefaultMaterial(MD_Surface))
+				{
+					Material->PostEditChange();
+				}
 			}
 
 			++ObjectIndex;
