Have a way to include exclude all when importing alembic cache
#fix Added top level check box to include / exclude all items in listview
#misc Fixed double clicking behaviour to toggle only selected item

diff --git a/Engine/Plugins/Experimental/AlembicImporter/Source/AlembicImporter/Private/AlembicImportOptions.cpp b/Engine/Plugins/Experimental/AlembicImporter/Source/AlembicImporter/Private/AlembicImportOptions.cpp
--- a/Engine/Plugins/Experimental/AlembicImporter/Source/AlembicImporter/Private/AlembicImportOptions.cpp
+++ b/Engine/Plugins/Experimental/AlembicImporter/Source/AlembicImporter/Private/AlembicImportOptions.cpp
@@ -88,8 +88,13 @@ void SAlembicImportOptions::Construct(const FArguments& InArgs)
 						SNew(SHeaderRow)
 
 						+ SHeaderRow::Column("ShouldImport")
-						.DefaultLabel(FText::FromString(TEXT("Include")))
 						.FillWidth(0.1f)
+						.DefaultLabel(FText::FromString(TEXT("Include")))
+						[
+							SNew(SCheckBox)
+							.HAlign(HAlign_Center)
+							.OnCheckStateChanged(this, &SAlembicImportOptions::OnToggleAllItems)
+						]
 
 						+ SHeaderRow::Column("TrackName")
 						.DefaultLabel(LOCTEXT("TrackNameHeader", "Track Name"))
@@ -157,11 +162,24 @@ bool SAlembicImportOptions::CanImport()  const
 	return true;
 }
 
+void SAlembicImportOptions::OnToggleAllItems(ECheckBoxState CheckType)
+{
+	/** Set all items to top level checkbox state */
+	for (FPolyMeshDataPtr& Item : PolyMeshData)
+	{
+		Item->PolyMesh->bShouldImport = CheckType == ECheckBoxState::Checked;
+	}
+}
+
 void SAlembicImportOptions::OnItemDoubleClicked(FPolyMeshDataPtr ClickedItem)
 {
+	/** Toggle state on / off for the selected list entry */
 	for (FPolyMeshDataPtr& Item : PolyMeshData)
 	{
-		Item->PolyMesh->bShouldImport = (Item == ClickedItem);
+		if (Item == ClickedItem)
+		{
+			Item->PolyMesh->bShouldImport = !Item->PolyMesh->bShouldImport;
+		}
 	}
 }
 
diff --git a/Engine/Plugins/Experimental/AlembicImporter/Source/AlembicImporter/Private/AlembicImportOptions.h b/Engine/Plugins/Experimental/AlembicImporter/Source/AlembicImporter/Private/AlembicImportOptions.h
--- a/Engine/Plugins/Experimental/AlembicImporter/Source/AlembicImporter/Private/AlembicImportOptions.h
+++ b/Engine/Plugins/Experimental/AlembicImporter/Source/AlembicImporter/Private/AlembicImportOptions.h
@@ -86,6 +86,7 @@ private:
 	TSharedRef<ITableRow> OnGenerateWidgetForList(FPolyMeshDataPtr InItem, const TSharedRef<STableViewBase>& OwnerTable);
 	bool CanImport() const;
 
+	void OnToggleAllItems(ECheckBoxState CheckType);
 	void OnItemDoubleClicked(FPolyMeshDataPtr ClickedItem);
 private:
 	UAbcImportSettings*	ImportSettings;
