Fixing compile errors after the last merge from main

diff --git a/Engine/Plugins/Experimental/AlembicImporter/Source/AlembicLibrary/Private/AbcImporter.cpp b/Engine/Plugins/Experimental/AlembicImporter/Source/AlembicLibrary/Private/AbcImporter.cpp
--- a/Engine/Plugins/Experimental/AlembicImporter/Source/AlembicLibrary/Private/AbcImporter.cpp
+++ b/Engine/Plugins/Experimental/AlembicImporter/Source/AlembicLibrary/Private/AbcImporter.cpp
@@ -646,7 +646,7 @@ TArray<UObject*> FAbcImporter::ImportAsSkeletalMesh(UObject* InParent, EObjectFl
 		FSkeletalMeshModel* ImportedModel = SkeletalMesh->GetImportedModel();
 		check(ImportedModel->LODModels.Num() == 0);
 		ImportedModel->LODModels.Empty();
-		ImportedModel->EmptyOriginalReductionSourceMeshData()
+		ImportedModel->EmptyOriginalReductionSourceMeshData();
 		ImportedModel->LODModels.Add(new FSkeletalMeshLODModel());
 		SkeletalMesh->ResetLODInfo();
 		SkeletalMesh->AddLODInfo();
diff --git a/Engine/Source/Runtime/Core/Public/HAL/LowLevelMemTracker.h b/Engine/Source/Runtime/Core/Public/HAL/LowLevelMemTracker.h
--- a/Engine/Source/Runtime/Core/Public/HAL/LowLevelMemTracker.h
+++ b/Engine/Source/Runtime/Core/Public/HAL/LowLevelMemTracker.h
@@ -94,10 +94,21 @@
 			}
 			return false;
 		}
+
+		// This is used by ensure to generate a bool per instance
+		// by passing a lambda which will uniquely instantiate the template.
+		template <typename Type>
+		bool TrueOnFirstCallOnly(const Type&)
+		{
+			static bool bValue = true;
+			bool Result = bValue;
+			bValue = false;
+			return Result;
+		}
 	}
 
 #if !USING_CODE_ANALYSIS
-	#define LLMTrueOnFirstCallOnly			UE4Asserts_Private::TrueOnFirstCallOnly([]{})
+	#define LLMTrueOnFirstCallOnly			LLMPrivate::TrueOnFirstCallOnly([]{})
 #else
 	#define LLMTrueOnFirstCallOnly			false
 #endif
