When including headers from a third party that is compiled with default visibility we need to include the headers as such

This is due to a weak symbol that wont be able to be overriden correctly

diff --git a/Engine/Plugins/Experimental/AlembicImporter/Source/AlembicLibrary/Private/AbcFile.cpp b/Engine/Plugins/Experimental/AlembicImporter/Source/AlembicLibrary/Private/AbcFile.cpp
--- a/Engine/Plugins/Experimental/AlembicImporter/Source/AlembicLibrary/Private/AbcFile.cpp
+++ b/Engine/Plugins/Experimental/AlembicImporter/Source/AlembicLibrary/Private/AbcFile.cpp
@@ -16,10 +16,12 @@
 
 #include "HAL/Platform.h"
 
+PRAGMA_DEFAULT_VISIBILITY_START
 THIRD_PARTY_INCLUDES_START
 #include <Alembic/AbcGeom/All.h>
 #include "Materials/MaterialInstance.h"
 THIRD_PARTY_INCLUDES_END
+PRAGMA_DEFAULT_VISIBILITY_START
 
 #define LOCTEXT_NAMESPACE "AbcFile"
 
@@ -489,4 +491,4 @@ UMaterialInterface** FAbcFile::GetMaterialByName(const FString& InMaterialName)
 	return MaterialMap.Find(InMaterialName);
 }
 
-#undef LOCTEXT_NAMESPACE // "AbcFile"
\ No newline at end of file
+#undef LOCTEXT_NAMESPACE // "AbcFile"
diff --git a/Engine/Plugins/Experimental/AlembicImporter/Source/AlembicLibrary/Private/AbcImportUtilities.cpp b/Engine/Plugins/Experimental/AlembicImporter/Source/AlembicLibrary/Private/AbcImportUtilities.cpp
--- a/Engine/Plugins/Experimental/AlembicImporter/Source/AlembicLibrary/Private/AbcImportUtilities.cpp
+++ b/Engine/Plugins/Experimental/AlembicImporter/Source/AlembicLibrary/Private/AbcImportUtilities.cpp
@@ -10,6 +10,7 @@
 #include "Windows/WindowsHWrapper.h"
 #endif
 
+PRAGMA_DEFAULT_VISIBILITY_START
 THIRD_PARTY_INCLUDES_START
 #include <Alembic/AbcCoreHDF5/All.h>
 #include <Alembic/AbcCoreOgawa/All.h>
@@ -20,6 +21,7 @@ THIRD_PARTY_INCLUDES_START
 #include <Alembic/AbcGeom/All.h>
 #include <Alembic/Abc/IObject.h>
 THIRD_PARTY_INCLUDES_END
+PRAGMA_DEFAULT_VISIBILITY_END
 
 #include "Rendering/SkeletalMeshLODModel.h"
 
diff --git a/Engine/Plugins/Experimental/AlembicImporter/Source/AlembicLibrary/Private/AbcImportUtilities.h b/Engine/Plugins/Experimental/AlembicImporter/Source/AlembicLibrary/Private/AbcImportUtilities.h
--- a/Engine/Plugins/Experimental/AlembicImporter/Source/AlembicLibrary/Private/AbcImportUtilities.h
+++ b/Engine/Plugins/Experimental/AlembicImporter/Source/AlembicLibrary/Private/AbcImportUtilities.h
@@ -8,11 +8,13 @@
 #include "Windows/WindowsHWrapper.h"
 #endif
 
+PRAGMA_DEFAULT_VISIBILITY_START
 THIRD_PARTY_INCLUDES_START
 #include <Alembic/AbcCoreAbstract/TimeSampling.h>
 #include <Alembic/Abc/All.h>
 #include <Alembic/AbcGeom/All.h>
 THIRD_PARTY_INCLUDES_END
+PRAGMA_DEFAULT_VISIBILITY_END
 
 #include "GeometryCache.h"
 #include "GeometryCacheTrackFlipbookAnimation.h"
diff --git a/Engine/Plugins/Experimental/AlembicImporter/Source/AlembicLibrary/Private/AbcImporter.cpp b/Engine/Plugins/Experimental/AlembicImporter/Source/AlembicLibrary/Private/AbcImporter.cpp
--- a/Engine/Plugins/Experimental/AlembicImporter/Source/AlembicLibrary/Private/AbcImporter.cpp
+++ b/Engine/Plugins/Experimental/AlembicImporter/Source/AlembicLibrary/Private/AbcImporter.cpp
@@ -6,6 +6,7 @@
 #include "Windows/WindowsHWrapper.h"
 #endif
 
+PRAGMA_DEFAULT_VISIBILITY_START
 THIRD_PARTY_INCLUDES_START
 #include <Alembic/AbcCoreHDF5/All.h>
 #include <Alembic/AbcCoreOgawa/All.h>
@@ -13,6 +14,7 @@ THIRD_PARTY_INCLUDES_START
 #include <Alembic/AbcCoreAbstract/TimeSampling.h>
 #include <Alembic/AbcCoreHDF5/All.h>
 THIRD_PARTY_INCLUDES_END
+PRAGMA_DEFAULT_VISIBILITY_END
 
 #include "Misc/Paths.h"
 #include "Misc/FeedbackContext.h"
diff --git a/Engine/Plugins/Experimental/AlembicImporter/Source/AlembicLibrary/Private/AbcPolyMesh.cpp b/Engine/Plugins/Experimental/AlembicImporter/Source/AlembicLibrary/Private/AbcPolyMesh.cpp
--- a/Engine/Plugins/Experimental/AlembicImporter/Source/AlembicLibrary/Private/AbcPolyMesh.cpp
+++ b/Engine/Plugins/Experimental/AlembicImporter/Source/AlembicLibrary/Private/AbcPolyMesh.cpp
@@ -7,10 +7,12 @@
 #include "Modules/ModuleManager.h"
 #include "MeshUtilities.h"
 
+PRAGMA_DEFAULT_VISIBILITY_START
 THIRD_PARTY_INCLUDES_START
 #include <Alembic/AbcCoreAbstract/TimeSampling.h>
 #include <Alembic/AbcGeom/Visibility.h>
 THIRD_PARTY_INCLUDES_END
+PRAGMA_DEFAULT_VISIBILITY_END
 
 static const ESampleReadFlags ReadAllFlags = ESampleReadFlags::Positions | ESampleReadFlags::Indices | ESampleReadFlags::UVs | ESampleReadFlags::Normals | ESampleReadFlags::Colors | ESampleReadFlags::MaterialIndices;
 
diff --git a/Engine/Plugins/Experimental/AlembicImporter/Source/AlembicLibrary/Private/AbcTransform.cpp b/Engine/Plugins/Experimental/AlembicImporter/Source/AlembicLibrary/Private/AbcTransform.cpp
--- a/Engine/Plugins/Experimental/AlembicImporter/Source/AlembicLibrary/Private/AbcTransform.cpp
+++ b/Engine/Plugins/Experimental/AlembicImporter/Source/AlembicLibrary/Private/AbcTransform.cpp
@@ -5,10 +5,12 @@
 
 #include "AbcFile.h"
 
+PRAGMA_DEFAULT_VISIBILITY_START
 THIRD_PARTY_INCLUDES_START
 #include <Alembic/AbcCoreAbstract/TimeSampling.h>
 #include <Alembic/AbcGeom/Visibility.h>
 THIRD_PARTY_INCLUDES_END
+PRAGMA_DEFAULT_VISIBILITY_END
 
 FAbcTransform::FAbcTransform(const Alembic::AbcGeom::IXform& InTransform, const FAbcFile* InFile, IAbcObject* InParent /*= nullptr*/) : IAbcObject(InTransform, InFile, InParent), Transform(InTransform), Schema(InTransform.getSchema())
 {
diff --git a/Engine/Plugins/Experimental/AlembicImporter/Source/AlembicLibrary/Private/EigenHelper.h b/Engine/Plugins/Experimental/AlembicImporter/Source/AlembicLibrary/Private/EigenHelper.h
--- a/Engine/Plugins/Experimental/AlembicImporter/Source/AlembicLibrary/Private/EigenHelper.h
+++ b/Engine/Plugins/Experimental/AlembicImporter/Source/AlembicLibrary/Private/EigenHelper.h
@@ -13,10 +13,12 @@
 	#pragma warning(push)
 	#pragma warning(disable:6294) // Ill-defined for-loop:  initial condition does not satisfy test.  Loop body not executed.
 #endif
+PRAGMA_DEFAULT_VISIBILITY_START
 THIRD_PARTY_INCLUDES_START
 	#include <Eigen/Dense>
 	#include <Eigen/SVD>
 THIRD_PARTY_INCLUDES_END
+PRAGMA_DEFAULT_VISIBILITY_END
 #if defined(_MSC_VER) && USING_CODE_ANALYSIS
 	#pragma warning(pop)
 #endif
diff --git a/Engine/Plugins/Experimental/AlembicImporter/Source/AlembicLibrary/Public/AbcFile.h b/Engine/Plugins/Experimental/AlembicImporter/Source/AlembicLibrary/Public/AbcFile.h
--- a/Engine/Plugins/Experimental/AlembicImporter/Source/AlembicLibrary/Public/AbcFile.h
+++ b/Engine/Plugins/Experimental/AlembicImporter/Source/AlembicLibrary/Public/AbcFile.h
@@ -9,12 +9,13 @@
 #include "Windows/WindowsHWrapper.h"
 #endif
 
+PRAGMA_DEFAULT_VISIBILITY_START
 THIRD_PARTY_INCLUDES_START
 #include <Alembic/AbcCoreFactory/IFactory.h>
 #include <Alembic/Abc/IArchive.h>
 #include <Alembic/Abc/IObject.h>
-
 THIRD_PARTY_INCLUDES_END
+PRAGMA_DEFAULT_VISIBILITY_END
 
 class IAbcObject;
 class FAbcPolyMesh;
@@ -137,4 +138,4 @@ protected:
 
 	/** Cached Mesh utilities ptr for normal calculations */
 	IMeshUtilities* MeshUtilities;
-};
\ No newline at end of file
+};
diff --git a/Engine/Plugins/Experimental/AlembicImporter/Source/AlembicLibrary/Public/AbcImporter.h b/Engine/Plugins/Experimental/AlembicImporter/Source/AlembicLibrary/Public/AbcImporter.h
--- a/Engine/Plugins/Experimental/AlembicImporter/Source/AlembicLibrary/Public/AbcImporter.h
+++ b/Engine/Plugins/Experimental/AlembicImporter/Source/AlembicLibrary/Public/AbcImporter.h
@@ -12,9 +12,11 @@
 #include "Windows/WindowsHWrapper.h"
 #endif
 
+PRAGMA_DEFAULT_VISIBILITY_START
 THIRD_PARTY_INCLUDES_START
 #include <Alembic/AbcGeom/All.h>
 THIRD_PARTY_INCLUDES_END
+PRAGMA_DEFAULT_VISIBILITY_END
 
 class UMaterial;
 class UStaticMesh;
@@ -212,4 +214,4 @@ private:
 
 	/** ABC file representation for currently opened filed */
 	class FAbcFile* AbcFile;
-};
\ No newline at end of file
+};
diff --git a/Engine/Plugins/Experimental/AlembicImporter/Source/AlembicLibrary/Public/AbcObject.h b/Engine/Plugins/Experimental/AlembicImporter/Source/AlembicLibrary/Public/AbcObject.h
--- a/Engine/Plugins/Experimental/AlembicImporter/Source/AlembicLibrary/Public/AbcObject.h
+++ b/Engine/Plugins/Experimental/AlembicImporter/Source/AlembicLibrary/Public/AbcObject.h
@@ -9,10 +9,12 @@
 #include "Windows/WindowsHWrapper.h"
 #endif
 
+PRAGMA_DEFAULT_VISIBILITY_START
 THIRD_PARTY_INCLUDES_START
 #include <Alembic/Abc/IObject.h>
 #include <Alembic/AbcGeom/All.h>
 THIRD_PARTY_INCLUDES_END
+PRAGMA_DEFAULT_VISIBILITY_END
 
 class FAbcFile;
 enum class EFrameReadFlags : uint8;
@@ -66,4 +68,4 @@ protected:
 	float FrameTimes[MaxNumberOfResidentSamples];
 	int32 ResidentSampleIndices[MaxNumberOfResidentSamples];
 	bool InUseSamples[MaxNumberOfResidentSamples];
-};
\ No newline at end of file
+};
diff --git a/Engine/Plugins/Media/ImgMedia/Source/OpenExrWrapper/Private/OpenExrWrapper.cpp b/Engine/Plugins/Media/ImgMedia/Source/OpenExrWrapper/Private/OpenExrWrapper.cpp
--- a/Engine/Plugins/Media/ImgMedia/Source/OpenExrWrapper/Private/OpenExrWrapper.cpp
+++ b/Engine/Plugins/Media/ImgMedia/Source/OpenExrWrapper/Private/OpenExrWrapper.cpp
@@ -5,6 +5,7 @@
 #include "Containers/UnrealString.h"
 #include "Modules/ModuleManager.h"
 
+PRAGMA_DEFAULT_VISIBILITY_START
 THIRD_PARTY_INCLUDES_START
 	#include "ImathBox.h"
 	#include "ImfHeader.h"
@@ -12,6 +13,7 @@ THIRD_PARTY_INCLUDES_START
 	#include "ImfCompressionAttribute.h"
 	#include "ImfStandardAttributes.h"
 THIRD_PARTY_INCLUDES_END
+PRAGMA_DEFAULT_VISIBILITY_END
 
 
 /* FOpenExr
diff --git a/Engine/Source/Runtime/Core/Public/Apple/ApplePlatformCompilerPreSetup.h b/Engine/Source/Runtime/Core/Public/Apple/ApplePlatformCompilerPreSetup.h
--- a/Engine/Source/Runtime/Core/Public/Apple/ApplePlatformCompilerPreSetup.h
+++ b/Engine/Source/Runtime/Core/Public/Apple/ApplePlatformCompilerPreSetup.h
@@ -79,3 +79,6 @@
 #define PRAGMA_DISABLE_OPTIMIZATION_ACTUAL _Pragma("clang optimize off")
 #define PRAGMA_ENABLE_OPTIMIZATION_ACTUAL  _Pragma("clang optimize on")
 #endif
+
+#define PRAGMA_DEFAULT_VISIBILITY_START _Pragma("GCC visibility push(default)")
+#define PRAGMA_DEFAULT_VISIBILITY_END   _Pragma("GCC visibility pop")
diff --git a/Engine/Source/Runtime/Core/Public/GenericPlatform/GenericPlatformCompilerPreSetup.h b/Engine/Source/Runtime/Core/Public/GenericPlatform/GenericPlatformCompilerPreSetup.h
--- a/Engine/Source/Runtime/Core/Public/GenericPlatform/GenericPlatformCompilerPreSetup.h
+++ b/Engine/Source/Runtime/Core/Public/GenericPlatform/GenericPlatformCompilerPreSetup.h
@@ -87,3 +87,11 @@
 #ifndef DEPRECATED_MACRO
 	#define DEPRECATED_MACRO(Version, Message) EMIT_CUSTOM_WARNING(Message " Please update your code to the new API before upgrading to the next release, otherwise your project will no longer compile.")
 #endif
+
+#ifndef PRAGMA_DEFAULT_VISIBILITY_START
+	#define PRAGMA_DEFAULT_VISIBILITY_START
+#endif
+
+#ifndef PRAGMA_DEFAULT_VISIBILITY_END
+	#define PRAGMA_DEFAULT_VISIBILITY_END
+#endif
