Alembic caches imported as static meshes skipping empty frames at start appear incorrect.
#fix make sure we setup correct import range for static meshes, used to read 2 frames and import the 2nd one rather than the 1st

diff --git a/Engine/Plugins/Experimental/AlembicImporter/Source/AlembicLibrary/Private/AbcFile.cpp b/Engine/Plugins/Experimental/AlembicImporter/Source/AlembicLibrary/Private/AbcFile.cpp
--- a/Engine/Plugins/Experimental/AlembicImporter/Source/AlembicLibrary/Private/AbcFile.cpp
+++ b/Engine/Plugins/Experimental/AlembicImporter/Source/AlembicLibrary/Private/AbcFile.cpp
@@ -124,12 +124,12 @@ EAbcImportError FAbcFile::Import(UAbcImportSettings* InImportSettings)
 
 	if (ImportSettings->ImportType == EAlembicImportType::StaticMesh)
 	{
-		EndFrameIndex = StartFrameIndex + 1;
+		EndFrameIndex = StartFrameIndex;
 	}
 
 	int32 FrameSpan = EndFrameIndex - StartFrameIndex;
 	// If Start==End or Start > End output error message due to invalid frame span
-	if (FrameSpan <= 0)
+	if (FrameSpan <= 0 && ImportSettings->ImportType != EAlembicImportType::StaticMesh)
 	{
 		TSharedRef<FTokenizedMessage> Message = FTokenizedMessage::Create(EMessageSeverity::Error, FText::Format(LOCTEXT("NoFramesForMeshObject", "Invalid frame range specified {0} - {1}."), FText::FromString(FString::FromInt(StartFrameIndex)), FText::FromString(FString::FromInt(EndFrameIndex))));
 		FAbcImportLogger::AddImportMessage(Message);
