Add 0 entries to enums to prevent blueprint compiler warnings

Enums affected:
- EBaseCalculationType

#rnx

diff --git a/Engine/Plugins/Experimental/AlembicImporter/Source/AlembicLibrary/Public/AbcImportSettings.h b/Engine/Plugins/Experimental/AlembicImporter/Source/AlembicLibrary/Public/AbcImportSettings.h
--- a/Engine/Plugins/Experimental/AlembicImporter/Source/AlembicLibrary/Public/AbcImportSettings.h
+++ b/Engine/Plugins/Experimental/AlembicImporter/Source/AlembicLibrary/Public/AbcImportSettings.h
@@ -21,6 +21,7 @@ enum class EAlembicImportType : uint8
 UENUM(BlueprintType)
 enum class EBaseCalculationType : uint8
 {
+	None = 0 UMETA(Hidden),
 	/** Determines the number of bases that should be used with the given percentage*/
 	PercentageBased = 1,
 	/** Set a fixed number of bases to import*/
