Fixed crash when re-importing Alembic SkeletalMesh.

diff --git a/Engine/Plugins/Experimental/AlembicImporter/Source/AlembicLibrary/Private/AbcImporter.cpp b/Engine/Plugins/Experimental/AlembicImporter/Source/AlembicLibrary/Private/AbcImporter.cpp
--- a/Engine/Plugins/Experimental/AlembicImporter/Source/AlembicLibrary/Private/AbcImporter.cpp
+++ b/Engine/Plugins/Experimental/AlembicImporter/Source/AlembicLibrary/Private/AbcImporter.cpp
@@ -712,82 +712,89 @@ TArray<UObject*> FAbcImporter::ImportAsSkeletalMesh(UObject* InParent, EObjectFl
 		uint32 WedgeOffset = 0;
 		uint32 VertexOffset = 0;
 
-		for (FCompressedAbcData& CompressedData : CompressedMeshData)
 		{
-			FAbcMeshSample* AverageSample = CompressedData.AverageSample;
-
-			if (CompressedData.BaseSamples.Num() > 0)
+#if WITH_EDITOR
+			// When ScopedPostEditChange goes out of scope, it will call SkeletalMesh->PostEditChange()
+			// while preventing any call to that within the scope
+			FScopedSkeletalMeshPostEditChange ScopedPostEditChange(SkeletalMesh);
+#endif
+			for (FCompressedAbcData& CompressedData : CompressedMeshData)
 			{
-				const int32 NumBases = CompressedData.BaseSamples.Num();
-				int32 NumUsedBases = 0;
+				FAbcMeshSample* AverageSample = CompressedData.AverageSample;
 
-				const int32 NumIndices = CompressedData.AverageSample->Indices.Num();
-
-				for (int32 BaseIndex = 0; BaseIndex < NumBases; ++BaseIndex)
+				if (CompressedData.BaseSamples.Num() > 0)
 				{
-					FAbcMeshSample* BaseSample = CompressedData.BaseSamples[BaseIndex];
+					const int32 NumBases = CompressedData.BaseSamples.Num();
+					int32 NumUsedBases = 0;
 
-					//AbcImporterUtilities::CalculateNormalsWithSmoothingGroups(BaseSample, AverageSample->SmoothingGroupIndices, AverageSample->NumSmoothingGroups);
+					const int32 NumIndices = CompressedData.AverageSample->Indices.Num();
 
-					// Create new morph target with name based on object and base index
-					UMorphTarget* MorphTarget = NewObject<UMorphTarget>(SkeletalMesh, FName(*FString::Printf(TEXT("Base_%i_%i"), BaseIndex, ObjectIndex)));
+					for (int32 BaseIndex = 0; BaseIndex < NumBases; ++BaseIndex)
+					{
+						FAbcMeshSample* BaseSample = CompressedData.BaseSamples[BaseIndex];
 
-					// Setup morph target vertices directly
-					TArray<FMorphTargetDelta> MorphDeltas;
-					GenerateMorphTargetVertices(BaseSample, MorphDeltas, AverageSample, WedgeOffset, MorphTargetVertexRemapping, UsedVertexIndicesForMorphs, VertexOffset, WedgeOffset);
-					MorphTarget->PopulateDeltas(MorphDeltas, 0, LODModel.Sections);
+						//AbcImporterUtilities::CalculateNormalsWithSmoothingGroups(BaseSample, AverageSample->SmoothingGroupIndices, AverageSample->NumSmoothingGroups);
 
-					const float PercentageOfVerticesInfluences = ((float)MorphTarget->MorphLODModels[0].Vertices.Num() / (float)NumIndices) * 100.0f;
-					if (PercentageOfVerticesInfluences > ImportSettings->CompressionSettings.MinimumNumberOfVertexInfluencePercentage)
-					{
-						SkeletalMesh->RegisterMorphTarget(MorphTarget);
-						MorphTarget->MarkPackageDirty();
-
-						// Set up curves
-						const TArray<float>& CurveValues = CompressedData.CurveValues[BaseIndex];
-						const TArray<float>& TimeValues = CompressedData.TimeValues[BaseIndex];
-						// Morph target stuffies
-						FString CurveName = FString::Printf(TEXT("Base_%i_%i"), BaseIndex, ObjectIndex);
-						FName ConstCurveName = *CurveName;
-
-						// Sets up the morph target curves with the sample values and time keys
-						SetupMorphTargetCurves(Skeleton, ConstCurveName, Sequence, CurveValues, TimeValues);
-					}
-					else
-					{
-						MorphTarget->MarkPendingKill();
+						// Create new morph target with name based on object and base index
+						UMorphTarget* MorphTarget = NewObject<UMorphTarget>(SkeletalMesh, FName(*FString::Printf(TEXT("Base_%i_%i"), BaseIndex, ObjectIndex)));
+
+						// Setup morph target vertices directly
+						TArray<FMorphTargetDelta> MorphDeltas;
+						GenerateMorphTargetVertices(BaseSample, MorphDeltas, AverageSample, WedgeOffset, MorphTargetVertexRemapping, UsedVertexIndicesForMorphs, VertexOffset, WedgeOffset);
+						MorphTarget->PopulateDeltas(MorphDeltas, 0, LODModel.Sections);
+
+						const float PercentageOfVerticesInfluences = ((float)MorphTarget->MorphLODModels[0].Vertices.Num() / (float)NumIndices) * 100.0f;
+						if (PercentageOfVerticesInfluences > ImportSettings->CompressionSettings.MinimumNumberOfVertexInfluencePercentage)
+						{
+							SkeletalMesh->RegisterMorphTarget(MorphTarget);
+							MorphTarget->MarkPackageDirty();
+
+							// Set up curves
+							const TArray<float>& CurveValues = CompressedData.CurveValues[BaseIndex];
+							const TArray<float>& TimeValues = CompressedData.TimeValues[BaseIndex];
+							// Morph target stuffies
+							FString CurveName = FString::Printf(TEXT("Base_%i_%i"), BaseIndex, ObjectIndex);
+							FName ConstCurveName = *CurveName;
+
+							// Sets up the morph target curves with the sample values and time keys
+							SetupMorphTargetCurves(Skeleton, ConstCurveName, Sequence, CurveValues, TimeValues);
+						}
+						else
+						{
+							MorphTarget->MarkPendingKill();
+						}
 					}
 				}
-			}
 
-			Sequence->RawCurveData.RemoveRedundantKeys();
+				Sequence->RawCurveData.RemoveRedundantKeys();
 
-			WedgeOffset += CompressedData.AverageSample->Indices.Num();
-			VertexOffset += CompressedData.AverageSample->Vertices.Num();
+				WedgeOffset += CompressedData.AverageSample->Indices.Num();
+				VertexOffset += CompressedData.AverageSample->Vertices.Num();
 
-			const uint32 NumMaterials = CompressedData.MaterialNames.Num();
-			for (uint32 MaterialIndex = 0; MaterialIndex < NumMaterials; ++MaterialIndex)
-			{
-				const FString& MaterialName = CompressedData.MaterialNames[MaterialIndex];
-				UMaterialInterface* Material = RetrieveMaterial(MaterialName, InParent, Flags);
-				SkeletalMesh->Materials.Add(FSkeletalMaterial(Material, true));
-				if (Material != UMaterial::GetDefaultMaterial(MD_Surface))
+				const uint32 NumMaterials = CompressedData.MaterialNames.Num();
+				for (uint32 MaterialIndex = 0; MaterialIndex < NumMaterials; ++MaterialIndex)
 				{
-					Material->PostEditChange();
+					const FString& MaterialName = CompressedData.MaterialNames[MaterialIndex];
+					UMaterialInterface* Material = RetrieveMaterial(MaterialName, InParent, Flags);
+					SkeletalMesh->Materials.Add(FSkeletalMaterial(Material, true));
+					if (Material != UMaterial::GetDefaultMaterial(MD_Surface))
+					{
+						Material->PostEditChange();
+					}
 				}
+
+				++ObjectIndex;
 			}
 
-			++ObjectIndex;
-		}
+			// Set recompute tangent flag on skeletal mesh sections
+			for (FSkelMeshSection& Section : LODModel.Sections)
+			{
+				Section.bRecomputeTangent = true;
+			}
 
-		// Set recompute tangent flag on skeletal mesh sections
-		for (FSkelMeshSection& Section : LODModel.Sections)
-		{
-			Section.bRecomputeTangent = true;
+			SkeletalMesh->CalculateInvRefMatrices();
 		}
 
-		SkeletalMesh->CalculateInvRefMatrices();
-		SkeletalMesh->PostEditChange();
 		SkeletalMesh->MarkPackageDirty();
 
 		// Retrieve the name mapping container
