PR #6044: Fixed crash when color in Alembic asset is stored per primitive. (Contributed by [REDACTED_PERSON])

diff --git a/Engine/Plugins/Experimental/AlembicImporter/Source/AlembicLibrary/Private/AbcImportUtilities.h b/Engine/Plugins/Experimental/AlembicImporter/Source/AlembicLibrary/Private/AbcImportUtilities.h
--- a/Engine/Plugins/Experimental/AlembicImporter/Source/AlembicLibrary/Private/AbcImportUtilities.h
+++ b/Engine/Plugins/Experimental/AlembicImporter/Source/AlembicLibrary/Private/AbcImportUtilities.h
@@ -149,6 +149,46 @@ namespace AbcImporterUtilities
 		InOutData = NewData;
 	}
 
+	template<typename T> void ExpandPrimitiveAttributeArray(const TArray<uint32>& InFaceCounts, TArray<T>& InOutArray)
+	{
+		int32 NumIndices = 0;
+		for (const uint32 NumIndicesForFace : InFaceCounts)
+		{
+			NumIndices += NumIndicesForFace;
+		}
+
+		TArray<T> NewArray;
+		NewArray.Reserve(NumIndices);
+
+		int32 PrimitiveIndex = 0;
+		for (const uint32 NumIndicesForFace : InFaceCounts)
+		{
+			T data = InOutArray[PrimitiveIndex];
+			if (NumIndicesForFace > 3)
+			{
+				// Triangle 0
+				NewArray.Add(data);
+				NewArray.Add(data);
+				NewArray.Add(data);
+
+				// Triangle 1
+				NewArray.Add(data);
+				NewArray.Add(data);
+				NewArray.Add(data);
+			}
+			else
+			{
+				NewArray.Add(data);
+				NewArray.Add(data);
+				NewArray.Add(data);
+			}
+
+			PrimitiveIndex++;
+		}
+
+		InOutArray = NewArray;
+	}
+
 	template<typename T> void ProcessVertexAttributeArray(const TArray<uint32>& InIndices, const TArray<uint32>& InFaceCounts, const bool bTriangulation, const uint32 NumVertices, TArray<T>& InOutArray)
 	{
 		// Expand using the vertex indices (if num entries == num vertices)
@@ -156,6 +196,11 @@ namespace AbcImporterUtilities
 		{
 			ExpandVertexAttributeArray(InIndices, InOutArray);
 		}
+		// Expand using the primitive indices (if num entries == num faces)
+		else if (InOutArray.Num() != InIndices.Num() && InOutArray.Num() == InFaceCounts.Num())
+		{
+			ExpandPrimitiveAttributeArray(InFaceCounts, InOutArray);
+		}
 		// Otherwise the attributes are stored per face, so triangulate if the faces contain quads
 		else if (bTriangulation)
 		{
