Fix crash when importing Alembic as Static Mesh that has degenerated faces.

diff --git a/Engine/Plugins/Experimental/AlembicImporter/Source/AlembicLibrary/Private/AbcImporter.cpp b/Engine/Plugins/Experimental/AlembicImporter/Source/AlembicLibrary/Private/AbcImporter.cpp
--- a/Engine/Plugins/Experimental/AlembicImporter/Source/AlembicLibrary/Private/AbcImporter.cpp
+++ b/Engine/Plugins/Experimental/AlembicImporter/Source/AlembicLibrary/Private/AbcImporter.cpp
@@ -1275,16 +1275,28 @@ void FAbcImporter::GenerateMeshDescriptionFromSample(const FAbcMeshSample* Sampl
 		VertexPositions[VertexID] = FVector(Position);
 	}
 
+	uint32 VertexIndices[3];
 	uint32 TriangleCount = Sample->Indices.Num() / 3;
 	for (uint32 TriangleIndex = 0; TriangleIndex < TriangleCount; ++TriangleIndex)
 	{
+		const uint32 IndiceIndex0 = TriangleIndex * 3;
+		VertexIndices[0] = Sample->Indices[IndiceIndex0];
+		VertexIndices[1] = Sample->Indices[IndiceIndex0 + 1];
+		VertexIndices[2] = Sample->Indices[IndiceIndex0 + 2];
+
+		// Skip degenerated triangle
+		if (VertexIndices[0] == VertexIndices[1] || VertexIndices[1] == VertexIndices[2] || VertexIndices[0] == VertexIndices[2])
+		{
+			continue;
+		}
+
 		TArray<FVertexInstanceID> CornerVertexInstanceIDs;
 		CornerVertexInstanceIDs.SetNum(3);
 		FVertexID CornerVertexIDs[3];
 		for (int32 Corner = 0; Corner < 3; ++Corner)
 		{
-			uint32 IndiceIndex = (TriangleIndex * 3) + Corner;
-			uint32 VertexIndex = Sample->Indices[IndiceIndex];
+			uint32 IndiceIndex = IndiceIndex0 + Corner;
+			uint32 VertexIndex = VertexIndices[Corner];
 			const FVertexID VertexID(VertexIndex);
 			const FVertexInstanceID VertexInstanceID = MeshDescription->CreateVertexInstance(VertexID);
 
