Copied 11443501 from Dev-VirtualProduction.

Added support for materials for GeometryCacheAbcFile.

diff --git a/Engine/Plugins/Experimental/AlembicImporter/Source/AlembicLibrary/Private/AbcFile.cpp b/Engine/Plugins/Experimental/AlembicImporter/Source/AlembicLibrary/Private/AbcFile.cpp
--- a/Engine/Plugins/Experimental/AlembicImporter/Source/AlembicLibrary/Private/AbcFile.cpp
+++ b/Engine/Plugins/Experimental/AlembicImporter/Source/AlembicLibrary/Private/AbcFile.cpp
@@ -291,6 +291,25 @@ EAbcImportError FAbcFile::Import(UAbcImportSettings* InImportSettings)
 		}
 	}
 
+	// Populate the list of unique face set names from the meshes that should be imported regardless of the import material settings
+	bool bRequiresDefaultMaterial = false;
+	for (FAbcPolyMesh* PolyMesh : PolyMeshes)
+	{
+		if (PolyMesh->bShouldImport)
+		{
+			for (const FString& FaceSetName : PolyMesh->FaceSetNames)
+			{
+				UniqueFaceSetNames.AddUnique(FaceSetName);
+			}
+			bRequiresDefaultMaterial |= PolyMesh->FaceSetNames.Num() == 0;
+		}
+	}
+
+	if (bRequiresDefaultMaterial)
+	{
+		UniqueFaceSetNames.Insert(TEXT("DefaultMaterial"), 0);
+	}
+
 	return AbcImportError_NoError;
 }
 
diff --git a/Engine/Plugins/Experimental/AlembicImporter/Source/AlembicLibrary/Private/AbcImportUtilities.cpp b/Engine/Plugins/Experimental/AlembicImporter/Source/AlembicLibrary/Private/AbcImportUtilities.cpp
--- a/Engine/Plugins/Experimental/AlembicImporter/Source/AlembicLibrary/Private/AbcImportUtilities.cpp
+++ b/Engine/Plugins/Experimental/AlembicImporter/Source/AlembicLibrary/Private/AbcImportUtilities.cpp
@@ -3,8 +3,11 @@
 #include "AbcImportUtilities.h"
 #include "Stats/StatsMisc.h"
 
+#include "AbcFile.h"
 #include "AbcImporter.h"
 #include "AbcPolyMesh.h"
+#include "AssetRegistryModule.h"
+#include "Materials/Material.h"
 
 #if PLATFORM_WINDOWS
 #include "Windows/WindowsHWrapper.h"
@@ -1546,4 +1549,54 @@ void AbcImporterUtilities::MergePolyMeshesToMeshData(int32 FrameIndex, int32 Fra
 	AbcImporterUtilities::GeometryCacheDataForMeshSample(MeshData, &MergedSample, 0);
 }
 
+UMaterialInterface* AbcImporterUtilities::RetrieveMaterial(FAbcFile& AbcFile, const FString& MaterialName, UObject* InParent, EObjectFlags Flags)
+{
+	UMaterialInterface* Material = nullptr;
+	UMaterialInterface** CachedMaterial = AbcFile.GetMaterialByName(MaterialName);
+	if (CachedMaterial)
+	{
+		Material = *CachedMaterial;
+		// Material could have been deleted if we're overriding/reimporting an asset
+		if (Material->IsValidLowLevel())
+		{
+			if (Material->GetOuter() == GetTransientPackage())
+			{
+				UMaterial* ExistingTypedObject = FindObject<UMaterial>(InParent, *MaterialName);
+				if (!ExistingTypedObject)
+				{
+					// This is in for safety, as we do not expect this to happen
+					UObject* ExistingObject = FindObject<UObject>(InParent, *MaterialName);
+					if (ExistingObject)
+					{
+						return nullptr;
+					}
+
+					Material->Rename(*MaterialName, InParent);				
+					Material->SetFlags(Flags);
+					FAssetRegistryModule::AssetCreated(Material);
+				}
+				else
+				{
+					ExistingTypedObject->PreEditChange(nullptr);
+					Material = ExistingTypedObject;
+				}
+			}
+		}
+		else
+		{
+			// In this case recreate the material
+			Material = NewObject<UMaterial>(InParent, *MaterialName);
+			Material->SetFlags(Flags);
+			FAssetRegistryModule::AssetCreated(Material);
+		}
+	}
+	else
+	{
+		Material = UMaterial::GetDefaultMaterial(MD_Surface);
+		check(Material);
+	}
+
+	return Material;
+}
+
 #undef LOCTEXT_NAMESPACE // "AbcImporterUtilities"
diff --git a/Engine/Plugins/Experimental/AlembicImporter/Source/AlembicLibrary/Private/AbcImportUtilities.h b/Engine/Plugins/Experimental/AlembicImporter/Source/AlembicLibrary/Private/AbcImportUtilities.h
--- a/Engine/Plugins/Experimental/AlembicImporter/Source/AlembicLibrary/Private/AbcImportUtilities.h
+++ b/Engine/Plugins/Experimental/AlembicImporter/Source/AlembicLibrary/Private/AbcImportUtilities.h
@@ -32,7 +32,9 @@ THIRD_PARTY_INCLUDES_END
 #include "AbcImportLogger.h"
 #include "AbcImportSettings.h"
 
+class FAbcFile;
 class FAbcPolyMesh;
+class UMaterialInterface;
 struct FAbcMeshSample;
 struct FCompressedAbcData;
 
@@ -339,4 +341,7 @@ namespace AbcImporterUtilities
 	 * @param bConstantTopology		Flag to indicate if the merged PolyMeshes has constant topology
 	 */
 	void MergePolyMeshesToMeshData(int32 FrameIndex, int32 FrameStart, const TArray<FAbcPolyMesh*>& PolyMeshes, const TArray<FString>& UniqueFaceSetNames, FGeometryCacheMeshData& MeshData, int32& PreviousNumVertices, bool& bConstantTopology);
+
+	/** Retrieves a material from an AbcFile according to the given name and resaves it into the parent package */
+	UMaterialInterface* RetrieveMaterial(FAbcFile& AbcFile, const FString& MaterialName, UObject* InParent, EObjectFlags Flags);
 }
diff --git a/Engine/Plugins/Experimental/AlembicImporter/Source/AlembicLibrary/Private/AbcImporter.cpp b/Engine/Plugins/Experimental/AlembicImporter/Source/AlembicLibrary/Private/AbcImporter.cpp
--- a/Engine/Plugins/Experimental/AlembicImporter/Source/AlembicLibrary/Private/AbcImporter.cpp
+++ b/Engine/Plugins/Experimental/AlembicImporter/Source/AlembicLibrary/Private/AbcImporter.cpp
@@ -33,16 +33,12 @@ THIRD_PARTY_INCLUDES_END
 #include "Animation/AnimSequence.h"
 #include "Rendering/SkeletalMeshModel.h"
 
-
 #include "AbcImportUtilities.h"
 #include "Utils.h"
 
 #include "MeshUtilities.h"
-#include "MaterialUtilities.h"
-#include "Materials/MaterialInstance.h"
-
-#include "Runtime/Engine/Classes/Materials/MaterialInterface.h"
-#include "Runtime/Engine/Public/MaterialCompiler.h"
+#include "Materials/Material.h"
+#include "Modules/ModuleManager.h"
 
 #include "Async/ParallelFor.h"
 
@@ -51,7 +47,6 @@ THIRD_PARTY_INCLUDES_END
 #include "AbcAssetImportData.h"
 #include "AbcFile.h"
 
-#include "AssetRegistryModule.h"
 #include "AnimationUtils.h"
 #include "ComponentReregisterContext.h"
 #include "GeometryCacheCodecV1.h"
@@ -216,7 +211,7 @@ UStaticMesh* FAbcImporter::CreateStaticMeshFromSample(UObject* InParent, const F
 			UMaterialInterface* Material = nullptr;
 			if (FaceSetNames.IsValidIndex(MaterialIndex))
 			{
-				Material = RetrieveMaterial(FaceSetNames[MaterialIndex], InParent, Flags);
+				Material = AbcImporterUtilities::RetrieveMaterial(*AbcFile, FaceSetNames[MaterialIndex], InParent, Flags);
 
 				if (Material != UMaterial::GetDefaultMaterial(MD_Surface))
 				{
@@ -388,29 +383,8 @@ UGeometryCache* FAbcImporter::ImportAsGeometryCache(UObject* InParent, EObjectFl
 				FScopedSlowTask SlowTask((ImportSettings->SamplingSettings.FrameEnd + 1) - ImportSettings->SamplingSettings.FrameStart, FText::FromString(FString(TEXT("Importing Frames"))));
 				SlowTask.MakeDialog(true);
 
-				// Need to get all face sets here -> material names?
-				TArray<FString> UniqueFaceSetNames;
-
+				const TArray<FString>& UniqueFaceSetNames = AbcFile->GetUniqueFaceSetNames();
 				const TArray<FAbcPolyMesh*>& PolyMeshes = AbcFile->GetPolyMeshes();
-				bool bRequiresDefaultMaterial = false;
-				for (FAbcPolyMesh* PolyMesh : PolyMeshes)
-				{
-					if (PolyMesh->bShouldImport)
-					{
-						for (const FString& FaceSetName : PolyMesh->FaceSetNames)
-						{
-							UniqueFaceSetNames.AddUnique(FaceSetName);
-						}
-					
-						bRequiresDefaultMaterial |= PolyMesh->FaceSetNames.Num() == 0;
-					}
-				}
-
-				if (bRequiresDefaultMaterial)
-				{
-					UniqueFaceSetNames.Insert(TEXT("DefaultMaterial"), 0 );
-				}
-
 				
 				const int32 NumTracks = Tracks.Num();
 				int32 PreviousNumVertices = 0;
@@ -434,7 +408,7 @@ UGeometryCache* FAbcImporter::ImportAsGeometryCache(UObject* InParent, EObjectFl
 				// Now add materials for all the face set names
 				for (const FString& FaceSetName : UniqueFaceSetNames)
 				{
-					UMaterialInterface* Material = RetrieveMaterial(FaceSetName, InParent, Flags);
+					UMaterialInterface* Material = AbcImporterUtilities::RetrieveMaterial(*AbcFile, FaceSetName, InParent, Flags);
 					GeometryCache->Materials.Add((Material != nullptr) ? Material : DefaultMaterial);		
 
 					if (Material != UMaterial::GetDefaultMaterial(MD_Surface))
@@ -467,7 +441,7 @@ UGeometryCache* FAbcImporter::ImportAsGeometryCache(UObject* InParent, EObjectFl
 							UMaterialInterface* Material = nullptr;
 							if (PolyMesh->FaceSetNames.IsValidIndex(MaterialIndex))
 							{
-								Material = RetrieveMaterial(PolyMesh->FaceSetNames[MaterialIndex], InParent, Flags);
+								Material = AbcImporterUtilities::RetrieveMaterial(*AbcFile, PolyMesh->FaceSetNames[MaterialIndex], InParent, Flags);
 								if (Material != UMaterial::GetDefaultMaterial(MD_Surface))
 								{
 									Material->PostEditChange();
@@ -730,7 +704,7 @@ TArray<UObject*> FAbcImporter::ImportAsSkeletalMesh(UObject* InParent, EObjectFl
 				for (uint32 MaterialIndex = 0; MaterialIndex < NumMaterials; ++MaterialIndex)
 				{
 					const FString& MaterialName = CompressedData.MaterialNames[MaterialIndex];
-					UMaterialInterface* Material = RetrieveMaterial(MaterialName, InParent, Flags);
+					UMaterialInterface* Material = AbcImporterUtilities::RetrieveMaterial(*AbcFile, MaterialName, InParent, Flags);
 					SkeletalMesh->Materials.Add(FSkeletalMaterial(Material, true));
 					if (Material != UMaterial::GetDefaultMaterial(MD_Surface))
 					{
@@ -1552,56 +1526,6 @@ void FAbcImporter::GenerateMorphTargetVertices(FAbcMeshSample* BaseSample, TArra
 	}
 }
 
-UMaterialInterface* FAbcImporter::RetrieveMaterial(const FString& MaterialName, UObject* InParent, EObjectFlags Flags)
-{
-	UMaterialInterface* Material = nullptr;
-	UMaterialInterface** CachedMaterial = AbcFile->GetMaterialByName(MaterialName);
-	if (CachedMaterial)
-	{
-		Material = *CachedMaterial;
-		// Material could have been deleted if we're overriding/reimporting an asset
-		if (Material->IsValidLowLevel())
-		{
-			if (Material->GetOuter() == GetTransientPackage())
-			{
-				UMaterial* ExistingTypedObject = FindObject<UMaterial>(InParent, *MaterialName);
-				if (!ExistingTypedObject)
-				{
-					// This is in for safety, as we do not expect this to happen
-					UObject* ExistingObject = FindObject<UObject>(InParent, *MaterialName);
-					if (ExistingObject)
-					{
-						return nullptr;
-					}
-
-					Material->Rename(*MaterialName, InParent);				
-					Material->SetFlags(Flags);
-					FAssetRegistryModule::AssetCreated(Material);
-				}
-				else
-				{
-					ExistingTypedObject->PreEditChange(nullptr);
-					Material = ExistingTypedObject;
-				}
-			}
-		}
-		else
-		{
-			// In this case recreate the material
-			Material = NewObject<UMaterial>(InParent, *MaterialName);
-			Material->SetFlags(Flags);
-			FAssetRegistryModule::AssetCreated(Material);
-		}
-	}
-	else
-	{
-		Material = UMaterial::GetDefaultMaterial(MD_Surface);
-		check(Material);
-	}
-
-	return Material;
-}
-
 FCompressedAbcData::~FCompressedAbcData()
 {
 	delete AverageSample;
diff --git a/Engine/Plugins/Experimental/AlembicImporter/Source/AlembicLibrary/Private/AbcUtilities.cpp b/Engine/Plugins/Experimental/AlembicImporter/Source/AlembicLibrary/Private/AbcUtilities.cpp
--- a/Engine/Plugins/Experimental/AlembicImporter/Source/AlembicLibrary/Private/AbcUtilities.cpp
+++ b/Engine/Plugins/Experimental/AlembicImporter/Source/AlembicLibrary/Private/AbcUtilities.cpp
@@ -4,23 +4,39 @@
 
 #include "AbcFile.h"
 #include "AbcImportUtilities.h"
-#include "AbcPolyMesh.h"
+#include "GeometryCache.h"
+#include "Materials/Material.h"
 
 void FAbcUtilities::GetFrameMeshData(FAbcFile& AbcFile, int32 FrameIndex, FGeometryCacheMeshData& OutMeshData)
 {
 	AbcFile.ReadFrame(FrameIndex, EFrameReadFlags::ApplyMatrix, 0);
 
-	// TODO: Get the FaceSetNames of the PolyMeshes in the Alembic
-	TArray<FString> UniqueFaceSetNames;
-	UniqueFaceSetNames.Insert(TEXT("DefaultMaterial"), 0);
-
 	FGeometryCacheMeshData MeshData;
 	int32 PreviousNumVertices = 0;
 	bool bConstantTopology = false;
 
-	AbcImporterUtilities::MergePolyMeshesToMeshData(FrameIndex, 0, AbcFile.GetPolyMeshes(), UniqueFaceSetNames, MeshData, PreviousNumVertices, bConstantTopology);
+	AbcImporterUtilities::MergePolyMeshesToMeshData(FrameIndex, 0, AbcFile.GetPolyMeshes(), AbcFile.GetUniqueFaceSetNames(), MeshData, PreviousNumVertices, bConstantTopology);
 
 	OutMeshData = MoveTemp(MeshData);
 
 	AbcFile.CleanupFrameData(0);
 }
+
+void FAbcUtilities::SetupGeometryCacheMaterials(FAbcFile& AbcFile, UGeometryCache* GeometryCache, UObject* Package)
+{
+	GeometryCache->Materials.Reset();
+
+	const EObjectFlags Flags = RF_Public | RF_Standalone | RF_Transactional;
+	UMaterial* DefaultMaterial = UMaterial::GetDefaultMaterial(MD_Surface);
+
+	for (const FString& FaceSetName : AbcFile.GetUniqueFaceSetNames())
+	{
+		UMaterialInterface* Material = AbcImporterUtilities::RetrieveMaterial(AbcFile, FaceSetName, Package, Flags);
+		GeometryCache->Materials.Add((Material != nullptr) ? Material : DefaultMaterial);		
+
+		if (Material != UMaterial::GetDefaultMaterial(MD_Surface))
+		{
+			Material->PostEditChange();
+		}
+	}
+}
diff --git a/Engine/Plugins/Experimental/AlembicImporter/Source/AlembicLibrary/Public/AbcFile.h b/Engine/Plugins/Experimental/AlembicImporter/Source/AlembicLibrary/Public/AbcFile.h
--- a/Engine/Plugins/Experimental/AlembicImporter/Source/AlembicLibrary/Public/AbcFile.h
+++ b/Engine/Plugins/Experimental/AlembicImporter/Source/AlembicLibrary/Public/AbcFile.h
@@ -89,7 +89,8 @@ public:
 	void ReadFrame(int32 FrameIndex, const EFrameReadFlags InFlags, const int32 ReadIndex = INDEX_NONE);
 	/** Cleans up frame data. Must be called after ReadFrame when the frame data (with matching ReadIndex) is not needed anymore */
 	void CleanupFrameData(const int32 ReadIndex);
-
+	/** Returns the list of unique face set names from the meshes to be imported */
+	const TArray<FString>& GetUniqueFaceSetNames() const { return UniqueFaceSetNames; }
 protected:
 	void TraverseAbcHierarchy(const Alembic::Abc::IObject& InObject, IAbcObject* InParent);
 protected:
@@ -122,7 +123,8 @@ protected:
 	
 	/** Map of material created for the imported alembic file identified by material names */
 	TMap<FString, UMaterialInterface*> MaterialMap;
-	
+	TArray<FString> UniqueFaceSetNames;
+
 	/** Total (max) number of frames in the Alembic file */
 	int32 NumFrames;
 	/** Frames per second (retrieved and specified in top Alembic object) */
diff --git a/Engine/Plugins/Experimental/AlembicImporter/Source/AlembicLibrary/Public/AbcImporter.h b/Engine/Plugins/Experimental/AlembicImporter/Source/AlembicLibrary/Public/AbcImporter.h
--- a/Engine/Plugins/Experimental/AlembicImporter/Source/AlembicLibrary/Public/AbcImporter.h
+++ b/Engine/Plugins/Experimental/AlembicImporter/Source/AlembicLibrary/Public/AbcImporter.h
@@ -186,9 +186,6 @@ private:
 	/** Generates and populate a FMeshDescription instance from the given sample*/
 	void GenerateMeshDescriptionFromSample(const FAbcMeshSample* Sample, FMeshDescription* MeshDescription, UStaticMesh* StaticMesh);
 	
-	/** Retrieves a material according to the given name and resaves it into the parent package*/
-	UMaterialInterface* RetrieveMaterial(const FString& MaterialName, UObject* InParent, EObjectFlags Flags );
-		
 	/** Compresses the imported animation data, returns true if compression was successful and compressed data was populated */
 	const bool CompressAnimationDataUsingPCA(const FAbcCompressionSettings& InCompressionSettings, const bool bRunComparison = false);	
 	/** Performs the actual SVD compression to retrieve the bases and weights used to set up the Skeletal mesh's morph targets */
diff --git a/Engine/Plugins/Experimental/AlembicImporter/Source/AlembicLibrary/Public/AbcUtilities.h b/Engine/Plugins/Experimental/AlembicImporter/Source/AlembicLibrary/Public/AbcUtilities.h
--- a/Engine/Plugins/Experimental/AlembicImporter/Source/AlembicLibrary/Public/AbcUtilities.h
+++ b/Engine/Plugins/Experimental/AlembicImporter/Source/AlembicLibrary/Public/AbcUtilities.h
@@ -12,4 +12,7 @@ class ALEMBICLIBRARY_API FAbcUtilities
 public:
 	/** Populates a FGeometryCacheMeshData instance (with merged meshes) for the given frame of an Alembic */
 	static void GetFrameMeshData(FAbcFile& AbcFile, int32 FrameIndex, FGeometryCacheMeshData& OutMeshData);
+
+	/** Sets up materials from an AbcFile to a GeometryCache and moves them into the given package */
+	static void SetupGeometryCacheMaterials(FAbcFile& AbcFile, class UGeometryCache* GeometryCache, UObject* Package);
 };
\ No newline at end of file
diff --git a/Engine/Plugins/Experimental/GeometryCacheAbcFile/Source/GeometryCacheAbcFile/Private/GeometryCacheAbcFileComponent.cpp b/Engine/Plugins/Experimental/GeometryCacheAbcFile/Source/GeometryCacheAbcFile/Private/GeometryCacheAbcFileComponent.cpp
--- a/Engine/Plugins/Experimental/GeometryCacheAbcFile/Source/GeometryCacheAbcFile/Private/GeometryCacheAbcFileComponent.cpp
+++ b/Engine/Plugins/Experimental/GeometryCacheAbcFile/Source/GeometryCacheAbcFile/Private/GeometryCacheAbcFileComponent.cpp
@@ -62,6 +62,9 @@ void UGeometryCacheAbcFileComponent::ReloadAbcFile()
 	{
 		// Also store the number of frames in the cache
 		GeometryCache->SetFrameStartEnd(0, AbcFileTrack->GetEndFrameIndex());
+
+		// Setup the materials from the AbcFile to the GeometryCache
+		AbcFileTrack->SetupGeometryCacheMaterials(GeometryCache);
 	}
 	else
 	{
@@ -69,6 +72,9 @@ void UGeometryCacheAbcFileComponent::ReloadAbcFile()
 		AlembicFilePath.FilePath.Empty();
 	}
 
+	ClearTrackData();
+	SetupTrackData();
+
 	MarkRenderStateDirty();
 }
 
diff --git a/Engine/Plugins/Experimental/GeometryCacheAbcFile/Source/GeometryCacheAbcFile/Private/GeometryCacheTrackAbcFile.cpp b/Engine/Plugins/Experimental/GeometryCacheAbcFile/Source/GeometryCacheAbcFile/Private/GeometryCacheTrackAbcFile.cpp
--- a/Engine/Plugins/Experimental/GeometryCacheAbcFile/Source/GeometryCacheAbcFile/Private/GeometryCacheTrackAbcFile.cpp
+++ b/Engine/Plugins/Experimental/GeometryCacheAbcFile/Source/GeometryCacheAbcFile/Private/GeometryCacheTrackAbcFile.cpp
@@ -7,9 +7,12 @@
 #include "AbcImportSettings.h"
 #include "AbcUtilities.h"
 #include "Framework/Notifications/NotificationManager.h"
+#include "GeometryCache.h"
 #include "GeometryCacheHelpers.h"
 #include "Logging/LogCategory.h"
 #include "Logging/LogMacros.h"
+#include "Misc/Paths.h"
+#include "PackageTools.h"
 #include "Widgets/Notifications/SNotificationList.h"
 
 DEFINE_LOG_CATEGORY_STATIC(LogGeometryCacheAbcFile, Log, All);
@@ -217,4 +220,20 @@ bool UGeometryCacheTrackAbcFile::GetMeshData(int32 SampleIndex, FGeometryCacheMe
 	return false;
 }
 
+void UGeometryCacheTrackAbcFile::SetupGeometryCacheMaterials(UGeometryCache* GeometryCache)
+{
+	if (AbcFile)
+	{
+		// Create package where the materials will be saved into
+		static const FString DestinationPath(TEXT("/Game/GeometryCacheAbcFile/Materials"));
+		FString Name = FPaths::GetBaseFilename(SourceFile);
+		FString PackageName = UPackageTools::SanitizePackageName(FPaths::Combine(*DestinationPath, *Name, *Name));
+
+		UPackage* Package = CreatePackage(nullptr, *PackageName);
+		Package->FullyLoad();
+
+		FAbcUtilities::SetupGeometryCacheMaterials(*AbcFile, GeometryCache, Package);
+	}
+}
+
 #undef LOCTEXT_NAMESPACE
\ No newline at end of file
diff --git a/Engine/Plugins/Experimental/GeometryCacheAbcFile/Source/GeometryCacheAbcFile/Public/GeometryCacheTrackAbcFile.h b/Engine/Plugins/Experimental/GeometryCacheAbcFile/Source/GeometryCacheAbcFile/Public/GeometryCacheTrackAbcFile.h
--- a/Engine/Plugins/Experimental/GeometryCacheAbcFile/Source/GeometryCacheAbcFile/Public/GeometryCacheTrackAbcFile.h
+++ b/Engine/Plugins/Experimental/GeometryCacheAbcFile/Source/GeometryCacheAbcFile/Public/GeometryCacheTrackAbcFile.h
@@ -35,6 +35,8 @@ public:
 
 	bool GetMeshData(int32 SampleIndex, FGeometryCacheMeshData& OutMeshData);
 
+	void SetupGeometryCacheMaterials(class UGeometryCache* GeometryCache);
+
 private:
 	void Reset();
 	void ShowNotification(const FText& Text);
