Fixed Alembic Import Options window sizing and positioning.

diff --git a/Engine/Plugins/Experimental/AlembicImporter/Source/AlembicImporter/Private/AlembicImportFactory.cpp b/Engine/Plugins/Experimental/AlembicImporter/Source/AlembicImporter/Private/AlembicImportFactory.cpp
--- a/Engine/Plugins/Experimental/AlembicImporter/Source/AlembicImporter/Private/AlembicImportFactory.cpp
+++ b/Engine/Plugins/Experimental/AlembicImporter/Source/AlembicImporter/Private/AlembicImportFactory.cpp
@@ -7,12 +7,10 @@
 #include "EditorFramework/AssetImportData.h"
 #include "HAL/FileManager.h"
 #include "Framework/Application/SlateApplication.h"
-
-
 #include "Interfaces/IMainFrameModule.h"
+#include "Math/UnrealMathUtility.h"
 
 #include "AlembicImportOptions.h"
-
 #include "AlembicLibraryModule.h"
 #include "AbcImporter.h"
 #include "AbcImportLogger.h"
@@ -498,10 +496,12 @@ EReimportResult::Type UAlembicImportFactory::Reimport(UObject* Obj)
 
 void UAlembicImportFactory::ShowImportOptionsWindow(TSharedPtr<SAlembicImportOptions>& Options, FString FilePath, const FAbcImporter& Importer)
 {
+	// Window size computed from SAlembicImportOptions
+	const float WindowHeight = 500.f + FMath::Clamp(Importer.GetPolyMeshes().Num() * 16.f, 0.f, 250.f);
 
 	TSharedRef<SWindow> Window = SNew(SWindow)
 		.Title(LOCTEXT("WindowTitle", "Alembic Cache Import Options"))
-		.SizingRule(ESizingRule::Autosized);
+		.ClientSize(FVector2D(522.f, WindowHeight));
 
 	Window->SetContent
 		(
diff --git a/Engine/Plugins/Experimental/AlembicImporter/Source/AlembicImporter/Private/AlembicImportOptions.cpp b/Engine/Plugins/Experimental/AlembicImporter/Source/AlembicImporter/Private/AlembicImportOptions.cpp
--- a/Engine/Plugins/Experimental/AlembicImporter/Source/AlembicImporter/Private/AlembicImportOptions.cpp
+++ b/Engine/Plugins/Experimental/AlembicImporter/Source/AlembicImporter/Private/AlembicImportOptions.cpp
@@ -5,6 +5,7 @@
 #include "PropertyEditorModule.h"
 #include "IDetailsView.h"
 #include "AbcImportSettings.h"
+#include "Math/UnrealMathUtility.h"
 #include "Modules/ModuleManager.h"
 #include "Widgets/Views/SListView.h"
 #include "Widgets/Layout/SUniformGridPanel.h"
@@ -25,6 +26,7 @@ void SAlembicImportOptions::Construct(const FArguments& InArgs)
 	FDetailsViewArgs DetailsViewArgs;
 	DetailsViewArgs.bAllowSearch = false;
 	DetailsViewArgs.NameAreaSettings = FDetailsViewArgs::HideNameArea;
+	DetailsViewArgs.ColumnWidth = 0.5f;
 	DetailsView = PropertyEditorModule.CreateDetailView(DetailsViewArgs);
 	DetailsView->SetObject(ImportSettings);
 	
@@ -33,6 +35,9 @@ void SAlembicImportOptions::Construct(const FArguments& InArgs)
 		PolyMeshData.Add(FPolyMeshDataPtr(new FPolyMeshData(PolyMesh)));
 	}
 
+	static const float MaxDesiredHeight = 250.f;
+	float MinDesiredHeight = FMath::Clamp(InArgs._PolyMeshes.Num() * 16.f, 0.f, MaxDesiredHeight);
+
 	this->ChildSlot
 	[
 		SNew(SVerticalBox)
@@ -75,7 +80,8 @@ void SAlembicImportOptions::Construct(const FArguments& InArgs)
 			[
 				SNew(SBox)				
 				.MinDesiredWidth(512.0f)
-				.MaxDesiredHeight(350.0f)
+				.MinDesiredHeight(MinDesiredHeight)
+				.MaxDesiredHeight(MaxDesiredHeight)
 				[
 					SNew(SListView<FPolyMeshDataPtr>)
 					.ItemHeight(24)						
