In AlembicImportFactory, use the import settings provided by the script (through the AssetImportTask) when in unattended mode.

diff --git a/Engine/Plugins/Experimental/AlembicImporter/Source/AlembicImporter/Private/AlembicImportFactory.cpp b/Engine/Plugins/Experimental/AlembicImporter/Source/AlembicImporter/Private/AlembicImportFactory.cpp
--- a/Engine/Plugins/Experimental/AlembicImporter/Source/AlembicImporter/Private/AlembicImportFactory.cpp
+++ b/Engine/Plugins/Experimental/AlembicImporter/Source/AlembicImporter/Private/AlembicImportFactory.cpp
@@ -1,6 +1,7 @@
 // Copyright Epic Games, Inc. All Rights Reserved.
 
 #include "AlembicImportFactory.h"
+#include "AssetImportTask.h"
 #include "Engine/StaticMesh.h"
 #include "Engine/SkeletalMesh.h"
 #include "Editor.h"
@@ -95,6 +96,14 @@ UObject* UAlembicImportFactory::FactoryCreateFile(UClass* InClass, UObject* InPa
 		// Set whether or not the user canceled
 		bOutOperationCanceled = !Options->ShouldImport();
 	}
+	else
+	{
+		UAbcImportSettings* ScriptedSettings = AssetImportTask ? Cast<UAbcImportSettings>(AssetImportTask->Options) : nullptr;
+		if (ScriptedSettings)
+		{
+			ImportSettings = ScriptedSettings;
+		}
+	}
 
 	// Set up message log page name to separate different assets
 	const FString PageName = "Importing " + InName.ToString() + ".abc";
