Fixed crash when opening AnimSequence imported from Alembic (as SkeletalMesh).

diff --git a/Engine/Plugins/Experimental/AlembicImporter/Source/AlembicLibrary/Private/AbcFile.cpp b/Engine/Plugins/Experimental/AlembicImporter/Source/AlembicLibrary/Private/AbcFile.cpp
--- a/Engine/Plugins/Experimental/AlembicImporter/Source/AlembicLibrary/Private/AbcFile.cpp
+++ b/Engine/Plugins/Experimental/AlembicImporter/Source/AlembicLibrary/Private/AbcFile.cpp
@@ -177,6 +177,7 @@ EAbcImportError FAbcFile::Import(UAbcImportSettings* InImportSettings)
 	}
 
 	SecondsPerFrame = TimeStep;
+	FramesPerSecond = TimeStep > 0.f ? 1 / TimeStep : 30;
 	ImportLength = FrameSpan * TimeStep;
 
 	// Calculate time offset from start of import animation range
@@ -485,6 +486,11 @@ const float FAbcFile::GetImportLength() const
 	return ImportLength;
 }
 
+const int32 FAbcFile::GetImportNumFrames() const
+{
+	return EndFrameIndex - StartFrameIndex;
+}
+
 const int32 FAbcFile::GetFramerate() const
 {
 	return FramesPerSecond;
diff --git a/Engine/Plugins/Experimental/AlembicImporter/Source/AlembicLibrary/Private/AbcImporter.cpp b/Engine/Plugins/Experimental/AlembicImporter/Source/AlembicLibrary/Private/AbcImporter.cpp
--- a/Engine/Plugins/Experimental/AlembicImporter/Source/AlembicLibrary/Private/AbcImporter.cpp
+++ b/Engine/Plugins/Experimental/AlembicImporter/Source/AlembicLibrary/Private/AbcImporter.cpp
@@ -636,6 +636,10 @@ TArray<UObject*> FAbcImporter::ImportAsSkeletalMesh(UObject* InParent, EObjectFl
 		Sequence->SequenceLength = AbcFile->GetImportLength();
 		Sequence->ImportFileFramerate = AbcFile->GetFramerate();
 		Sequence->ImportResampleFramerate = AbcFile->GetFramerate();
+
+		// Add 1 because of the way AnimSequence computes its framerate
+		Sequence->SetRawNumberOfFrame(AbcFile->GetImportNumFrames() + 1);
+
 		int32 ObjectIndex = 0;
 		uint32 TriangleOffset = 0;
 		uint32 WedgeOffset = 0;
diff --git a/Engine/Plugins/Experimental/AlembicImporter/Source/AlembicLibrary/Public/AbcFile.h b/Engine/Plugins/Experimental/AlembicImporter/Source/AlembicLibrary/Public/AbcFile.h
--- a/Engine/Plugins/Experimental/AlembicImporter/Source/AlembicLibrary/Public/AbcFile.h
+++ b/Engine/Plugins/Experimental/AlembicImporter/Source/AlembicLibrary/Public/AbcFile.h
@@ -65,6 +65,7 @@ public:
 	const int32 GetMaxFrameIndex() const;
 	const float GetImportTimeOffset() const;
 	const float GetImportLength() const;
+	const int32 GetImportNumFrames() const;
 	const int32 GetFramerate() const;
 	const FBoxSphereBounds& GetArchiveBounds() const;
 	const bool ContainsHeterogeneousMeshes() const;
