//UE5/Main - Run Incremental EngineTest Geometry Win64Editor - Assertion failed: (Index >= 0) & (Index < ArrayNum)
#fix ensure that the AnimModel is notifies to be 'populated' matching the Alembic importer with the FBX one. This ensures the resampling frame-rate is correctly populated.
#fix ensure that the Alembic SK anim importer generates coarse editor-only bone animation data

diff --git a/Engine/Plugins/Experimental/AlembicImporter/Source/AlembicLibrary/Private/AbcImporter.cpp b/Engine/Plugins/Experimental/AlembicImporter/Source/AlembicLibrary/Private/AbcImporter.cpp
--- a/Engine/Plugins/Experimental/AlembicImporter/Source/AlembicLibrary/Private/AbcImporter.cpp
+++ b/Engine/Plugins/Experimental/AlembicImporter/Source/AlembicLibrary/Private/AbcImporter.cpp
@@ -723,12 +723,16 @@ TArray<UObject*> FAbcImporter::ImportAsSkeletalMesh(UObject* InParent, EObjectFl
 
 				FRawAnimSequenceTrack RootBoneTrack;
 				RootBoneTrack.PosKeys.Reserve(NumSamples);
-				RootBoneTrack.RotKeys.Add(FQuat::Identity); // At least one rotation key is required for the track to be valid
+				RootBoneTrack.RotKeys.Reserve(NumSamples);
+				RootBoneTrack.ScaleKeys.Reserve(NumSamples);
 
 				for (int32 SampleIndex = 0; SampleIndex < NumSamples; ++SampleIndex)
 				{
 					const FVector SampleOffset = SamplesOffsets.GetValue()[SampleIndex];
 					RootBoneTrack.PosKeys.Add(SampleOffset);
+
+					RootBoneTrack.RotKeys.Add(FQuat::Identity);
+					RootBoneTrack.ScaleKeys.Add(FVector::OneVector);
 				}
 
 				const FReferenceSkeleton& RefSkeleton = SkeletalMesh->GetRefSkeleton();
@@ -750,6 +754,7 @@ TArray<UObject*> FAbcImporter::ImportAsSkeletalMesh(UObject* InParent, EObjectFl
 		SkeletalMesh->MarkPackageDirty();
 
 		Controller->UpdateCurveNamesFromSkeleton(Skeleton, ERawCurveTrackTypes::RCT_Float);
+		Controller->NotifyPopulated();
 
 		Controller->CloseBracket();
 
