Fixed crash in Alembic import as skeletal mesh with LWC enabled.

diff --git a/Engine/Plugins/Experimental/AlembicImporter/Source/AlembicLibrary/Private/AbcImporter.cpp b/Engine/Plugins/Experimental/AlembicImporter/Source/AlembicLibrary/Private/AbcImporter.cpp
--- a/Engine/Plugins/Experimental/AlembicImporter/Source/AlembicLibrary/Private/AbcImporter.cpp
+++ b/Engine/Plugins/Experimental/AlembicImporter/Source/AlembicLibrary/Private/AbcImporter.cpp
@@ -1060,8 +1060,8 @@ const bool FAbcImporter::CompressAnimationDataUsingPCA(const FAbcCompressionSett
 
 			// Set up average frame 
 			CompressedData.AverageSample = new FAbcMeshSample(MergedZeroFrameSample);
-			FMemory::Memcpy(CompressedData.AverageSample->Vertices.GetData(), AverageVertexData.GetData(), sizeof(FVector) * NumVertices);
-			FMemory::Memcpy(CompressedData.AverageSample->Normals.GetData(), AverageNormalData.GetData(), sizeof(FVector) * NumIndices);
+			FMemory::Memcpy(CompressedData.AverageSample->Vertices.GetData(), AverageVertexData.GetData(), sizeof(FVector3f) * NumVertices);
+			FMemory::Memcpy(CompressedData.AverageSample->Normals.GetData(), AverageNormalData.GetData(), sizeof(FVector3f) * NumIndices);
 
 			const float FrameStep = (MaxTime - MinTime) / (float)(NumSamples - 1);
 			AbcImporterUtilities::GenerateCompressedMeshData(CompressedData, NumUsedSingularValues, NumSamples, BasesMatrix, NormalsBasesMatrix, OutV, FrameStep, FMath::Max(MinTime, 0.0f));
@@ -1244,8 +1244,8 @@ const bool FAbcImporter::CompressAnimationDataUsingPCA(const FAbcCompressionSett
 				CompressedMeshData.AddDefaulted();
 				FCompressedAbcData& CompressedData = CompressedMeshData.Last();
 				CompressedData.AverageSample = new FAbcMeshSample(*PolyMeshesToCompress[MeshIndex]->GetTransformedFirstSample());
-				FMemory::Memcpy(CompressedData.AverageSample->Vertices.GetData(), AverageVertexData[MeshIndex].GetData(), sizeof(FVector) * NumVertices);
-				FMemory::Memcpy(CompressedData.AverageSample->Normals.GetData(), AverageNormalData[MeshIndex].GetData(), sizeof(FVector) * NumIndices);
+				FMemory::Memcpy(CompressedData.AverageSample->Vertices.GetData(), AverageVertexData[MeshIndex].GetData(), sizeof(FVector3f) * NumVertices);
+				FMemory::Memcpy(CompressedData.AverageSample->Normals.GetData(), AverageNormalData[MeshIndex].GetData(), sizeof(FVector3f) * NumIndices);
 
 				if ( InCompressionSettings.BaseCalculationType != EBaseCalculationType::NoCompression )
 				{
