PR #7880: Export some symbols from AlembicLibrary (Contributed by [REDACTED_PERSON])

diff --git a/Engine/Plugins/Experimental/AlembicImporter/Source/AlembicLibrary/Public/AbcObject.h b/Engine/Plugins/Experimental/AlembicImporter/Source/AlembicLibrary/Public/AbcObject.h
--- a/Engine/Plugins/Experimental/AlembicImporter/Source/AlembicLibrary/Public/AbcObject.h
+++ b/Engine/Plugins/Experimental/AlembicImporter/Source/AlembicLibrary/Public/AbcObject.h
@@ -24,7 +24,7 @@ enum class EFrameReadFlags : uint8;
 
 static const int32 MaxNumberOfResidentSamples = 8;
 
-class IAbcObject
+class ALEMBICLIBRARY_API IAbcObject
 {
 public:
 	IAbcObject(const Alembic::Abc::IObject& InObject, const FAbcFile* InFile, IAbcObject* InParent);
diff --git a/Engine/Plugins/Experimental/AlembicImporter/Source/AlembicLibrary/Public/AbcPolyMesh.h b/Engine/Plugins/Experimental/AlembicImporter/Source/AlembicLibrary/Public/AbcPolyMesh.h
--- a/Engine/Plugins/Experimental/AlembicImporter/Source/AlembicLibrary/Public/AbcPolyMesh.h
+++ b/Engine/Plugins/Experimental/AlembicImporter/Source/AlembicLibrary/Public/AbcPolyMesh.h
@@ -71,7 +71,7 @@ struct FAbcMeshSample
 	float SampleTime;
 };
 
-class FAbcPolyMesh : public IAbcObject
+class ALEMBICLIBRARY_API FAbcPolyMesh : public IAbcObject
 {
 public:
 	FAbcPolyMesh(const Alembic::AbcGeom::IPolyMesh& InPolyMesh, const FAbcFile* InFile, IAbcObject* InParent = nullptr);
diff --git a/Engine/Plugins/Experimental/AlembicImporter/Source/AlembicLibrary/Public/AbcTransform.h b/Engine/Plugins/Experimental/AlembicImporter/Source/AlembicLibrary/Public/AbcTransform.h
--- a/Engine/Plugins/Experimental/AlembicImporter/Source/AlembicLibrary/Public/AbcTransform.h
+++ b/Engine/Plugins/Experimental/AlembicImporter/Source/AlembicLibrary/Public/AbcTransform.h
@@ -4,7 +4,7 @@
 
 #include "AbcObject.h"
 
-class FAbcTransform : public IAbcObject
+class ALEMBICLIBRARY_API FAbcTransform : public IAbcObject
 {
 public:
 	FAbcTransform(const Alembic::AbcGeom::IXform& InTransform, const FAbcFile* InFile, IAbcObject* InParent = nullptr);
