Streaming GeometryCache:
* Fixed serialization issue with Store Imported Vertex Numbers.

diff --git a/Engine/Plugins/Experimental/AlembicImporter/Source/AlembicLibrary/Private/AbcUtilities.cpp b/Engine/Plugins/Experimental/AlembicImporter/Source/AlembicLibrary/Private/AbcUtilities.cpp
--- a/Engine/Plugins/Experimental/AlembicImporter/Source/AlembicLibrary/Private/AbcUtilities.cpp
+++ b/Engine/Plugins/Experimental/AlembicImporter/Source/AlembicLibrary/Private/AbcUtilities.cpp
@@ -3,6 +3,7 @@
 #include "AbcUtilities.h"
 
 #include "AbcFile.h"
+#include "AbcImportSettings.h"
 #include "AbcImportUtilities.h"
 #include "GeometryCache.h"
 #include "Materials/Material.h"
@@ -11,11 +12,15 @@ void FAbcUtilities::GetFrameMeshData(FAbcFile& AbcFile, int32 FrameIndex, FGeome
 {
 	AbcFile.ReadFrame(FrameIndex, EFrameReadFlags::ApplyMatrix, ConcurrencyIndex);
 
+
+	const UAbcImportSettings* ImportSettings = AbcFile.GetImportSettings();
+	checkSlow(ImportSettings);
+
 	FGeometryCacheMeshData MeshData;
 	int32 PreviousNumVertices = 0;
 	bool bConstantTopology = false;
-	const bool bUseVelocitiesAsMotionVectors = true;
-	const bool bStoreImportedVertexNumbers = false;
+	const bool bUseVelocitiesAsMotionVectors = ImportSettings->GeometryCacheSettings.MotionVectors == EAbcGeometryCacheMotionVectorsImport::ImportAbcVelocitiesAsMotionVectors;
+	const bool bStoreImportedVertexNumbers = ImportSettings->GeometryCacheSettings.bStoreImportedVertexNumbers;
 
 	AbcImporterUtilities::MergePolyMeshesToMeshData(FrameIndex, 0, AbcFile.GetSecondsPerFrame(), bUseVelocitiesAsMotionVectors,
 		AbcFile.GetPolyMeshes(), AbcFile.GetUniqueFaceSetNames(), MeshData, PreviousNumVertices, bConstantTopology, bStoreImportedVertexNumbers);
diff --git a/Engine/Plugins/Experimental/GeometryCache/Source/GeometryCache/Private/GeometryCacheMeshData.cpp b/Engine/Plugins/Experimental/GeometryCache/Source/GeometryCache/Private/GeometryCacheMeshData.cpp
--- a/Engine/Plugins/Experimental/GeometryCache/Source/GeometryCache/Private/GeometryCacheMeshData.cpp
+++ b/Engine/Plugins/Experimental/GeometryCache/Source/GeometryCache/Private/GeometryCacheMeshData.cpp
@@ -95,7 +95,7 @@ FArchive& operator<<(FArchive& Ar, FGeometryCacheMeshData& Mesh)
 
 			if (Mesh.VertexInfo.bHasImportedVertexNumbers)
 			{
-				Ar.Serialize(&Mesh.ImportedVertexNumbers, Mesh.ImportedVertexNumbers.Num() * Mesh.ImportedVertexNumbers.GetTypeSize());
+				Ar.Serialize(&Mesh.ImportedVertexNumbers[0], Mesh.ImportedVertexNumbers.Num() * Mesh.ImportedVertexNumbers.GetTypeSize());
 			}
 		  
 			if (Mesh.VertexInfo.bHasMotionVectors)
diff --git a/Engine/Plugins/Experimental/GeometryCacheAbcFile/Source/GeometryCacheAbcFile/Private/GeometryCacheAbcStream.cpp b/Engine/Plugins/Experimental/GeometryCacheAbcFile/Source/GeometryCacheAbcFile/Private/GeometryCacheAbcStream.cpp
--- a/Engine/Plugins/Experimental/GeometryCacheAbcFile/Source/GeometryCacheAbcFile/Private/GeometryCacheAbcStream.cpp
+++ b/Engine/Plugins/Experimental/GeometryCacheAbcFile/Source/GeometryCacheAbcFile/Private/GeometryCacheAbcStream.cpp
@@ -40,7 +40,7 @@ struct FGeometryCacheAbcStreamReadRequest
 // differences, etc.) replace the version GUID below with a new one.
 // In case of merge conflicts with DDC versions, you MUST generate a new GUID
 // and set this new GUID as the version.
-#define ABCSTREAM_DERIVED_DATA_VERSION TEXT("88025D2E38A54CF29FA5A6CAE686B013")
+#define ABCSTREAM_DERIVED_DATA_VERSION TEXT("75887BC18F674774A53F91E4A57709D6")
 
 class FAbcStreamDDCUtils
 {
