GeometryCache:
* Fixed importing of GeometryCache with empty frames.

diff --git a/Engine/Plugins/Experimental/AlembicImporter/Source/AlembicLibrary/Private/AbcImportUtilities.cpp b/Engine/Plugins/Experimental/AlembicImporter/Source/AlembicLibrary/Private/AbcImportUtilities.cpp
--- a/Engine/Plugins/Experimental/AlembicImporter/Source/AlembicLibrary/Private/AbcImportUtilities.cpp
+++ b/Engine/Plugins/Experimental/AlembicImporter/Source/AlembicLibrary/Private/AbcImportUtilities.cpp
@@ -1508,7 +1508,7 @@ void AbcImporterUtilities::ApplyConversion(TArray<FMatrix>& InOutMatrices, const
 void AbcImporterUtilities::GeometryCacheDataForMeshSample(FGeometryCacheMeshData &OutMeshData, const FAbcMeshSample* MeshSample,
 	const uint32 MaterialOffset, const float SecondsPerFrame, const bool bUseVelocitiesAsMotionVectors, const bool bStoreImportedVertexNumbers)
 {
-	OutMeshData.BoundingBox = FBox3f(MeshSample->Vertices);
+	OutMeshData.BoundingBox = MeshSample->Vertices.Num() ? FBox3f(MeshSample->Vertices) : FBox3f(EForceInit::ForceInit);
 
 	// We currently always have everything except motion vectors
 	// and tangents are auto-configure based on their presence in the mesh sample data
