Alembic Import:
* Allowed importing static mesh from Alembic to propagate matrix transform without the "Merge Meshes" option enabled.

diff --git a/Engine/Plugins/Experimental/AlembicImporter/Source/AlembicLibrary/Private/AbcImporter.cpp b/Engine/Plugins/Experimental/AlembicImporter/Source/AlembicLibrary/Private/AbcImporter.cpp
--- a/Engine/Plugins/Experimental/AlembicImporter/Source/AlembicLibrary/Private/AbcImporter.cpp
+++ b/Engine/Plugins/Experimental/AlembicImporter/Source/AlembicLibrary/Private/AbcImporter.cpp
@@ -324,7 +324,7 @@ const TArray<UStaticMesh*> FAbcImporter::ImportAsStaticMesh(UObject* InParent, E
 	};
 	
 
-	EFrameReadFlags ReadFlags = ( ImportSettings->StaticMeshSettings.bMergeMeshes && ImportSettings->StaticMeshSettings.bPropagateMatrixTransformations ? EFrameReadFlags::ApplyMatrix : EFrameReadFlags::None ) | EFrameReadFlags::ForceSingleThreaded;
+	EFrameReadFlags ReadFlags = ( ImportSettings->StaticMeshSettings.bPropagateMatrixTransformations ? EFrameReadFlags::ApplyMatrix : EFrameReadFlags::None ) | EFrameReadFlags::ForceSingleThreaded;
 	AbcFile->ProcessFrames(Func, ReadFlags);
 	
 	TArray<UObject*> Assets;
diff --git a/Engine/Plugins/Experimental/AlembicImporter/Source/AlembicLibrary/Private/AbcPolyMesh.cpp b/Engine/Plugins/Experimental/AlembicImporter/Source/AlembicLibrary/Private/AbcPolyMesh.cpp
--- a/Engine/Plugins/Experimental/AlembicImporter/Source/AlembicLibrary/Private/AbcPolyMesh.cpp
+++ b/Engine/Plugins/Experimental/AlembicImporter/Source/AlembicLibrary/Private/AbcPolyMesh.cpp
@@ -68,7 +68,7 @@ bool FAbcPolyMesh::ReadFirstFrame(const float InTime, const int32 FrameIndex)
 		// Compute tangents for mesh
 		AbcImporterUtilities::ComputeTangents(FirstSample, File->GetImportSettings()->NormalGenerationSettings.bIgnoreDegenerateTriangles, *File->GetMeshUtilities());
 
-		const bool bApplyTransformation = (File->GetImportSettings()->ImportType == EAlembicImportType::StaticMesh && File->GetImportSettings()->StaticMeshSettings.bMergeMeshes && File->GetImportSettings()->StaticMeshSettings.bPropagateMatrixTransformations) || (File->GetImportSettings()->ImportType == EAlembicImportType::Skeletal && File->GetImportSettings()->CompressionSettings.bBakeMatrixAnimation) || File->GetImportSettings()->ImportType == EAlembicImportType::GeometryCache;
+		const bool bApplyTransformation = (File->GetImportSettings()->ImportType == EAlembicImportType::StaticMesh && File->GetImportSettings()->StaticMeshSettings.bPropagateMatrixTransformations) || (File->GetImportSettings()->ImportType == EAlembicImportType::Skeletal && File->GetImportSettings()->CompressionSettings.bBakeMatrixAnimation) || File->GetImportSettings()->ImportType == EAlembicImportType::GeometryCache;
 
 		// Transform copy of the first sample 
 		TransformedFirstSample = new FAbcMeshSample(*FirstSample);
diff --git a/Engine/Plugins/Experimental/AlembicImporter/Source/AlembicLibrary/Public/AbcImportSettings.h b/Engine/Plugins/Experimental/AlembicImporter/Source/AlembicLibrary/Public/AbcImportSettings.h
--- a/Engine/Plugins/Experimental/AlembicImporter/Source/AlembicLibrary/Public/AbcImportSettings.h
+++ b/Engine/Plugins/Experimental/AlembicImporter/Source/AlembicLibrary/Public/AbcImportSettings.h
@@ -190,8 +190,8 @@ struct FAbcStaticMeshSettings
 	UPROPERTY(EditAnywhere, BlueprintReadWrite, Category = StaticMesh)
 	bool bMergeMeshes;
 
-	// This will, if applicable, apply matrix transformations to the meshes before merging
-	UPROPERTY(EditAnywhere, BlueprintReadWrite, Category = StaticMesh, meta=(editcondition="bMergeMeshes"))
+	// This will, if applicable, apply matrix transformations to the meshes
+	UPROPERTY(EditAnywhere, BlueprintReadWrite, Category = StaticMesh)
 	bool bPropagateMatrixTransformations;
 
 	// Flag for whether or not lightmap UVs should be generated
