Alembic Import:
* The import settings created from script are now used and displayed in the Import dialog. Previously, they were used only if the import was running unattended.

diff --git a/Engine/Plugins/Experimental/AlembicImporter/Source/AlembicImporter/Private/AlembicImportFactory.cpp b/Engine/Plugins/Experimental/AlembicImporter/Source/AlembicImporter/Private/AlembicImportFactory.cpp
--- a/Engine/Plugins/Experimental/AlembicImporter/Source/AlembicImporter/Private/AlembicImportFactory.cpp
+++ b/Engine/Plugins/Experimental/AlembicImporter/Source/AlembicImporter/Private/AlembicImportFactory.cpp
@@ -71,6 +71,13 @@ UObject* UAlembicImportFactory::FactoryCreateFile(UClass* InClass, UObject* InPa
 {
 	GEditor->GetEditorSubsystem<UImportSubsystem>()->BroadcastAssetPreImport(this, InClass, InParent, InName, TEXT("ABC"));
 
+	// Use (and show) the settings from the script if provided
+	UAbcImportSettings* ScriptedSettings = AssetImportTask ? Cast<UAbcImportSettings>(AssetImportTask->Options) : nullptr;
+	if (ScriptedSettings)
+	{
+		ImportSettings = ScriptedSettings;
+	}
+
 	FAbcImporter Importer;
 	EAbcImportError ErrorCode = Importer.OpenAbcFileForImport(Filename);
 	ImportSettings->bReimport = false;
@@ -83,9 +90,20 @@ UObject* UAlembicImportFactory::FactoryCreateFile(UClass* InClass, UObject* InPa
 		return nullptr;
 	}
 
-	// Reset (possible) changed frame start value 
-	ImportSettings->SamplingSettings.FrameStart = 0;
-	ImportSettings->SamplingSettings.FrameEnd = Importer.GetEndFrameIndex();
+	if (ImportSettings == UAbcImportSettings::Get())
+	{
+		// Reset (possible) changed frame start value 
+		ImportSettings->SamplingSettings.FrameStart = 0;
+		ImportSettings->SamplingSettings.FrameEnd = Importer.GetEndFrameIndex();
+	}
+	else
+	{
+		// Don't override the frame end value from the script except if it was unset
+		if (ImportSettings->SamplingSettings.FrameEnd == 0)
+		{
+			ImportSettings->SamplingSettings.FrameEnd = Importer.GetEndFrameIndex();
+		}
+	}
 
 	bOutOperationCanceled = false;
 
@@ -96,14 +114,6 @@ UObject* UAlembicImportFactory::FactoryCreateFile(UClass* InClass, UObject* InPa
 		// Set whether or not the user canceled
 		bOutOperationCanceled = !Options->ShouldImport();
 	}
-	else
-	{
-		UAbcImportSettings* ScriptedSettings = AssetImportTask ? Cast<UAbcImportSettings>(AssetImportTask->Options) : nullptr;
-		if (ScriptedSettings)
-		{
-			ImportSettings = ScriptedSettings;
-		}
-	}
 
 	// Set up message log page name to separate different assets
 	const FString PageName = "Importing " + InName.ToString() + ".abc";
