Alembic Import:
* Fixed some assets created from Alembic import (as skeletal mesh or static mesh) not appearing in Content Browser.

diff --git a/Engine/Plugins/Experimental/AlembicImporter/Source/AlembicImporter/Classes/AlembicImportFactory.h b/Engine/Plugins/Experimental/AlembicImporter/Source/AlembicImporter/Classes/AlembicImportFactory.h
--- a/Engine/Plugins/Experimental/AlembicImporter/Source/AlembicImporter/Classes/AlembicImportFactory.h
+++ b/Engine/Plugins/Experimental/AlembicImporter/Source/AlembicImporter/Classes/AlembicImportFactory.h
@@ -79,7 +79,7 @@ class ALEMBICIMPORTER_API UAlembicImportFactory : public UFactory, public FReimp
 	* @param Flags - Creation Flags for the GeometryCache asset
 	* @return UObject*
 	*/
-	UObject* ImportSkeletalMesh(FAbcImporter& Importer, UObject* InParent, EObjectFlags Flags);
+	TArray<UObject*> ImportSkeletalMesh(FAbcImporter& Importer, UObject* InParent, EObjectFlags Flags);
 
 	/**
 	* ReimportStaticMesh
diff --git a/Engine/Plugins/Experimental/AlembicImporter/Source/AlembicImporter/Private/AlembicImportFactory.cpp b/Engine/Plugins/Experimental/AlembicImporter/Source/AlembicImporter/Private/AlembicImportFactory.cpp
--- a/Engine/Plugins/Experimental/AlembicImporter/Source/AlembicImporter/Private/AlembicImportFactory.cpp
+++ b/Engine/Plugins/Experimental/AlembicImporter/Source/AlembicImporter/Private/AlembicImportFactory.cpp
@@ -2,6 +2,7 @@
 
 #include "AlembicImportFactory.h"
 #include "AssetImportTask.h"
+#include "AssetRegistry/AssetRegistryModule.h"
 #include "Engine/StaticMesh.h"
 #include "Engine/SkeletalMesh.h"
 #include "Editor.h"
@@ -156,11 +157,8 @@ UObject* UAlembicImportFactory::FactoryCreateFile(UClass* InClass, UObject* InPa
 			}
 			else if (ImportSettings->ImportType == EAlembicImportType::Skeletal)
 			{
-				UObject* SkeletalMesh = ImportSkeletalMesh(Importer, InParent, Flags);
-				if (SkeletalMesh)
-				{
-					ResultAssets.Add(SkeletalMesh);
-				}
+				TArray<UObject*> SkeletalMesh = ImportSkeletalMesh(Importer, InParent, Flags);
+				ResultAssets.Append(SkeletalMesh);
 			}
 		}
 
@@ -169,6 +167,7 @@ UObject* UAlembicImportFactory::FactoryCreateFile(UClass* InClass, UObject* InPa
 		{
 			if (Object)
 			{
+				FAssetRegistryModule::AssetCreated(Object);
 				GEditor->GetEditorSubsystem<UImportSubsystem>()->BroadcastAssetPostImport(this, Object);
 				Object->MarkPackageDirty();
 				Object->PostEditChange();
@@ -258,7 +257,7 @@ UObject* UAlembicImportFactory::ImportGeometryCache(FAbcImporter& Importer, UObj
 	}
 }
 
-UObject* UAlembicImportFactory::ImportSkeletalMesh(FAbcImporter& Importer, UObject* InParent, EObjectFlags Flags)
+TArray<UObject*> UAlembicImportFactory::ImportSkeletalMesh(FAbcImporter& Importer, UObject* InParent, EObjectFlags Flags)
 {
 	// Flush commands before importing
 	FlushRenderingCommands();
@@ -271,7 +270,7 @@ UObject* UAlembicImportFactory::ImportSkeletalMesh(FAbcImporter& Importer, UObje
 
 		if (!GeneratedObjects.Num())
 		{
-			return nullptr;
+			return {};
 		}
 
 		USkeletalMesh* SkeletalMesh = [&GeneratedObjects]()
@@ -316,14 +315,13 @@ UObject* UAlembicImportFactory::ImportSkeletalMesh(FAbcImporter& Importer, UObje
 			}
 		}
 		
-
-		return SkeletalMesh;
+		return GeneratedObjects;
 	}
 	else
 	{
-		// Not able to import a static mesh
+		// Not able to import as skeletal mesh
 		GEditor->GetEditorSubsystem<UImportSubsystem>()->BroadcastAssetPostImport(this, nullptr);
-		return nullptr;
+		return {};
 	}
 }
 
