Alembic Import:
* Added error message when trying to import HDF5 Alembic. The support for HDF5 has been removed since it's deprecated/legacy.

diff --git a/Engine/Plugins/Experimental/AlembicImporter/Source/AlembicImporter/Private/AlembicImportFactory.cpp b/Engine/Plugins/Experimental/AlembicImporter/Source/AlembicImporter/Private/AlembicImportFactory.cpp
--- a/Engine/Plugins/Experimental/AlembicImporter/Source/AlembicImporter/Private/AlembicImportFactory.cpp
+++ b/Engine/Plugins/Experimental/AlembicImporter/Source/AlembicImporter/Private/AlembicImportFactory.cpp
@@ -84,8 +84,14 @@ UObject* UAlembicImportFactory::FactoryCreateFile(UClass* InClass, UObject* InPa
 	ImportSettings->bReimport = false;
 	AdditionalImportedObjects.Empty();
 
+	// Set up message log page name to separate different assets
+	FText ImportingText = FText::Format(LOCTEXT("AbcFactoryImporting", "Importing {0}.abc"), FText::FromName(InName));
+	const FString& PageName = ImportingText.ToString();
+
 	if (ErrorCode != AbcImportError_NoError)
 	{
+		FAbcImportLogger::OutputMessages(PageName);
+
 		// Failed to read the file info, fail the import
 		GEditor->GetEditorSubsystem<UImportSubsystem>()->BroadcastAssetPostImport(this, nullptr);
 		return nullptr;
@@ -116,8 +122,6 @@ UObject* UAlembicImportFactory::FactoryCreateFile(UClass* InClass, UObject* InPa
 		bOutOperationCanceled = !Options->ShouldImport();
 	}
 
-	// Set up message log page name to separate different assets
-	const FString PageName = "Importing " + InName.ToString() + ".abc";
 
 	TArray<UObject*> ResultAssets;
 	if (!bOutOperationCanceled)
@@ -390,7 +394,9 @@ void UAlembicImportFactory::SetReimportPaths(UObject* Obj, const TArray<FString>
 EReimportResult::Type UAlembicImportFactory::Reimport(UObject* Obj)
 {
 	ImportSettings->bReimport = true;
-	const FString PageName = "Reimporting " + Obj->GetName() + ".abc";
+
+	FText ReimportingText = FText::Format(LOCTEXT("AbcFactoryReimporting", "Reimporting {0}.abc"), FText::FromString(Obj->GetName()));
+	const FString& PageName = ReimportingText.ToString();
 	if (Obj->GetClass() == UStaticMesh::StaticClass())
 	{
 		UStaticMesh* Mesh = Cast<UStaticMesh>(Obj);
diff --git a/Engine/Plugins/Experimental/AlembicImporter/Source/AlembicLibrary/Private/AbcFile.cpp b/Engine/Plugins/Experimental/AlembicImporter/Source/AlembicLibrary/Private/AbcFile.cpp
--- a/Engine/Plugins/Experimental/AlembicImporter/Source/AlembicLibrary/Private/AbcFile.cpp
+++ b/Engine/Plugins/Experimental/AlembicImporter/Source/AlembicLibrary/Private/AbcFile.cpp
@@ -81,6 +81,16 @@ EAbcImportError FAbcFile::Open()
 	Archive = Factory.getArchive(TCHAR_TO_ANSI(*FPaths::ConvertRelativePathToFull(FilePath)), CompressionType);
 	if (!Archive.valid())
 	{
+		TSharedPtr<FTokenizedMessage> Message;
+		if (CompressionType == Alembic::AbcCoreFactory::IFactory::CoreType::kHDF5)
+		{
+			Message = FTokenizedMessage::Create(EMessageSeverity::Error, LOCTEXT("AbcHDF5NotSupported", "HDF5 Alembic files are not supported. Please convert the file to Ogawa format before importing."));
+		}
+		else
+		{
+			Message = FTokenizedMessage::Create(EMessageSeverity::Error, LOCTEXT("AbcUnknownFormat", "Unknown file format: not a valid Alembic."));
+		}
+		FAbcImportLogger::AddImportMessage(Message.ToSharedRef());
 		return EAbcImportError::AbcImportError_InvalidArchive;
 	}
 
@@ -88,6 +98,8 @@ EAbcImportError FAbcFile::Open()
 	TopObject = Alembic::Abc::IObject(Archive, Alembic::Abc::kTop);
 	if (!TopObject.valid())
 	{
+		TSharedRef<FTokenizedMessage> Message = FTokenizedMessage::Create(EMessageSeverity::Error, LOCTEXT("AbcInvalidRoot", "Invalid root node: cannot proceed with import."));
+		FAbcImportLogger::AddImportMessage(Message);
 		return EAbcImportError::AbcImportError_NoValidTopObject;
 	}
 
