Alembic Import:
* Fixed issue with Alembic imported as Skeletal Mesh not animating after saving and reloading. The default mesh optimization setting was causing invalid morph targets.

diff --git a/Engine/Plugins/Importers/AlembicImporter/Source/AlembicLibrary/Private/AbcImporter.cpp b/Engine/Plugins/Importers/AlembicImporter/Source/AlembicLibrary/Private/AbcImporter.cpp
--- a/Engine/Plugins/Importers/AlembicImporter/Source/AlembicLibrary/Private/AbcImporter.cpp
+++ b/Engine/Plugins/Importers/AlembicImporter/Source/AlembicLibrary/Private/AbcImporter.cpp
@@ -579,7 +579,11 @@ TArray<UObject*> FAbcImporter::ImportAsSkeletalMesh(UObject* InParent, EObjectFl
 		ImportedModel->EmptyOriginalReductionSourceMeshData();
 		ImportedModel->LODModels.Add(new FSkeletalMeshLODModel());
 		SkeletalMesh->ResetLODInfo();
-		SkeletalMesh->AddLODInfo();
+
+		FSkeletalMeshLODInfo& NewLODInfo = SkeletalMesh->AddLODInfo();
+		NewLODInfo.ReductionSettings.NumOfTrianglesPercentage = 1.0f;
+		NewLODInfo.ReductionSettings.NumOfVertPercentage = 1.0f;
+		NewLODInfo.ReductionSettings.MaxDeviationPercentage = 0.0f;
 		FSkeletalMeshLODModel& LODModel = ImportedModel->LODModels[0];
 
 		const FMeshBoneInfo BoneInfo(FName(TEXT("RootBone"), FNAME_Add), TEXT("RootBone_Export"), INDEX_NONE);
