Alembic Import:
* Fixed Alembic import only saving the sampling import settings. Now the reimport will populate the options window with all the same settings used at import.

diff --git a/Engine/Plugins/Importers/AlembicImporter/Source/AlembicImporter/Private/AlembicImportFactory.cpp b/Engine/Plugins/Importers/AlembicImporter/Source/AlembicImporter/Private/AlembicImportFactory.cpp
--- a/Engine/Plugins/Importers/AlembicImporter/Source/AlembicImporter/Private/AlembicImportFactory.cpp
+++ b/Engine/Plugins/Importers/AlembicImporter/Source/AlembicImporter/Private/AlembicImportFactory.cpp
@@ -285,17 +285,17 @@ TArray<UObject*> UAlembicImportFactory::ImportSkeletalMesh(FAbcImporter& Importe
 			
 		if (SkeletalMesh)
 		{
-		// Setup asset import data
-		if (!SkeletalMesh->GetAssetImportData() || !SkeletalMesh->GetAssetImportData()->IsA<UAbcAssetImportData>())
-		{
-			SkeletalMesh->SetAssetImportData(NewObject<UAbcAssetImportData>(SkeletalMesh));
-		}
-		SkeletalMesh->GetAssetImportData()->Update(UFactory::CurrentFilename);
-		UAbcAssetImportData* AssetImportData = Cast<UAbcAssetImportData>(SkeletalMesh->GetAssetImportData());
-		if (AssetImportData)
-		{
-			Importer.UpdateAssetImportData(AssetImportData);
-		}
+			// Setup asset import data
+			if (!SkeletalMesh->GetAssetImportData() || !SkeletalMesh->GetAssetImportData()->IsA<UAbcAssetImportData>())
+			{
+				SkeletalMesh->SetAssetImportData(NewObject<UAbcAssetImportData>(SkeletalMesh));
+			}
+			SkeletalMesh->GetAssetImportData()->Update(UFactory::CurrentFilename);
+			UAbcAssetImportData* AssetImportData = Cast<UAbcAssetImportData>(SkeletalMesh->GetAssetImportData());
+			if (AssetImportData)
+			{
+				Importer.UpdateAssetImportData(AssetImportData);
+			}
 		}
 
 		UAnimSequence* AnimSequence = [&GeneratedObjects]()
@@ -654,12 +654,12 @@ EReimportResult::Type UAlembicImportFactory::ReimportSkeletalMesh(USkeletalMesh*
 	if (bShowOption)
 	{
 		TSharedPtr<SAlembicImportOptions> Options;
-	ShowImportOptionsWindow(Options, CurrentFilename, Importer);
+		ShowImportOptionsWindow(Options, CurrentFilename, Importer);
 
-	if (!Options->ShouldImport())
-	{
-		return EReimportResult::Cancelled;
-	}
+		if (!Options->ShouldImport())
+		{
+			return EReimportResult::Cancelled;
+		}
 	}
 
 	int32 NumThreads = 1;
@@ -736,6 +736,11 @@ EReimportResult::Type UAlembicImportFactory::ReimportSkeletalMesh(USkeletalMesh*
 void UAlembicImportFactory::PopulateOptionsWithImportData(UAbcAssetImportData* ImportData)
 {
 	ImportSettings->SamplingSettings = ImportData->SamplingSettings;
+	ImportSettings->NormalGenerationSettings = ImportData->NormalGenerationSettings;
+	ImportSettings->CompressionSettings = ImportData->CompressionSettings;
+	ImportSettings->StaticMeshSettings = ImportData->StaticMeshSettings;
+	ImportSettings->GeometryCacheSettings = ImportData->GeometryCacheSettings;
+	ImportSettings->ConversionSettings = ImportData->ConversionSettings;
 }
 
 EReimportResult::Type UAlembicImportFactory::ReimportStaticMesh(UStaticMesh* Mesh)
diff --git a/Engine/Plugins/Importers/AlembicImporter/Source/AlembicLibrary/Private/AbcImporter.cpp b/Engine/Plugins/Importers/AlembicImporter/Source/AlembicLibrary/Private/AbcImporter.cpp
--- a/Engine/Plugins/Importers/AlembicImporter/Source/AlembicLibrary/Private/AbcImporter.cpp
+++ b/Engine/Plugins/Importers/AlembicImporter/Source/AlembicLibrary/Private/AbcImporter.cpp
@@ -87,6 +87,11 @@ void FAbcImporter::UpdateAssetImportData(UAbcAssetImportData* AssetImportData)
 	}
 
 	AssetImportData->SamplingSettings = ImportSettings->SamplingSettings;
+	AssetImportData->NormalGenerationSettings = ImportSettings->NormalGenerationSettings;
+	AssetImportData->CompressionSettings = ImportSettings->CompressionSettings;
+	AssetImportData->StaticMeshSettings = ImportSettings->StaticMeshSettings;
+	AssetImportData->GeometryCacheSettings = ImportSettings->GeometryCacheSettings;
+	AssetImportData->ConversionSettings = ImportSettings->ConversionSettings;
 }
 
 void FAbcImporter::RetrieveAssetImportData(UAbcAssetImportData* AssetImportData)
diff --git a/Engine/Plugins/Importers/AlembicImporter/Source/AlembicLibrary/Public/AbcAssetImportData.h b/Engine/Plugins/Importers/AlembicImporter/Source/AlembicLibrary/Public/AbcAssetImportData.h
--- a/Engine/Plugins/Importers/AlembicImporter/Source/AlembicLibrary/Public/AbcAssetImportData.h
+++ b/Engine/Plugins/Importers/AlembicImporter/Source/AlembicLibrary/Public/AbcAssetImportData.h
@@ -21,4 +21,22 @@ public:
 
 	UPROPERTY()
 	FAbcSamplingSettings SamplingSettings;
+
+	UPROPERTY()
+	FAbcNormalGenerationSettings NormalGenerationSettings;
+
+	UPROPERTY()
+	FAbcMaterialSettings MaterialSettings;
+
+	UPROPERTY()
+	FAbcCompressionSettings CompressionSettings;
+
+	UPROPERTY()
+	FAbcStaticMeshSettings StaticMeshSettings;
+
+	UPROPERTY()
+	FAbcGeometryCacheSettings GeometryCacheSettings;
+
+	UPROPERTY()
+	FAbcConversionSettings ConversionSettings;
 };
