Adopting partner's improvements to alembic mesh reimporting functionality.
Not the 100% ideal way of addressing issues, but that would require a larger rework (which is covered in UETOOL-5021) -- accepting this in the interim.

diff --git a/Engine/Plugins/Importers/AlembicImporter/Source/AlembicImporter/Private/AlembicImportFactory.cpp b/Engine/Plugins/Importers/AlembicImporter/Source/AlembicImporter/Private/AlembicImportFactory.cpp
--- a/Engine/Plugins/Importers/AlembicImporter/Source/AlembicImporter/Private/AlembicImportFactory.cpp
+++ b/Engine/Plugins/Importers/AlembicImporter/Source/AlembicImporter/Private/AlembicImportFactory.cpp
@@ -22,6 +22,10 @@
 #include "GeometryCache.h"
 #include "Subsystems/AssetEditorSubsystem.h"
 
+#include "AI/Navigation/NavCollisionBase.h"
+#include "Editor/EditorPerProjectUserSettings.h"
+#include "ImportUtils/StaticMeshImportUtils.h"
+
 #define LOCTEXT_NAMESPACE "AlembicImportFactory"
 
 DEFINE_LOG_CATEGORY_STATIC(LogAlembic, Log, All);
@@ -70,6 +74,48 @@ bool UAlembicImportFactory::FactoryCanImport(const FString& Filename)
 
 UObject* UAlembicImportFactory::FactoryCreateFile(UClass* InClass, UObject* InParent, FName InName, EObjectFlags Flags, const FString& Filename, const TCHAR* Parms, FFeedbackContext* Warn, bool& bOutOperationCanceled)
 {
+	const bool bIsUnattended = (IsAutomatedImport()
+		|| FApp::IsUnattended()
+		|| IsRunningCommandlet()
+		|| GIsRunningUnattendedScript);
+	if (bIsUnattended)
+	{
+		bShowOption = false;
+	}
+
+	// Check if it's a re-import
+	if (InParent != nullptr)
+	{
+		UObject* ExistingObject =
+			StaticFindObject(UObject::StaticClass(), InParent, *InName.ToString());
+		if (ExistingObject)
+		{
+			// Use this class as no other re-import handler exist for Alembics, yet
+			FReimportHandler* ReimportHandler = this;
+			TArray<FString> Filenames;
+			Filenames.Add(UFactory::CurrentFilename);
+			// Set the new source path before starting the re-import
+			FReimportManager::Instance()->UpdateReimportPaths(ExistingObject, Filenames);
+			// Do the re-import and exit
+			const bool bIsAutomated = bIsUnattended;
+			const bool bShowNotification = !bIsAutomated;
+			const bool bAskForNewFileIfMissing = true;
+			const FString PreferredReimportFile;
+			const int32 SourceFileIndex = INDEX_NONE;
+			const bool bForceNewFile = false;
+			FReimportManager::Instance()->Reimport(
+				ExistingObject,
+				bAskForNewFileIfMissing,
+				bShowNotification,
+				PreferredReimportFile,
+				ReimportHandler,
+				SourceFileIndex,
+				bForceNewFile,
+				bIsAutomated);
+			return ExistingObject;
+		}
+	}
+
 	GEditor->GetEditorSubsystem<UImportSubsystem>()->BroadcastAssetPreImport(this, InClass, InParent, InName, TEXT("ABC"));
 
 	// Use (and show) the settings from the script if provided
@@ -114,7 +160,7 @@ UObject* UAlembicImportFactory::FactoryCreateFile(UClass* InClass, UObject* InPa
 
 	bOutOperationCanceled = false;
 
-	if (!GIsRunningUnattendedScript && bShowOption)
+	if (bShowOption)
 	{
 		TSharedPtr<SAlembicImportOptions> Options;
 		ShowImportOptionsWindow(Options, UFactory::CurrentFilename, Importer);
@@ -395,6 +441,14 @@ EReimportResult::Type UAlembicImportFactory::Reimport(UObject* Obj)
 {
 	ImportSettings->bReimport = true;
 
+	const bool bIsUnattended = (IsAutomatedImport()
+		|| FApp::IsUnattended()
+		|| IsRunningCommandlet()
+		|| GIsRunningUnattendedScript);
+	const bool bShowDialogAtReimport =
+		GetDefault<UEditorPerProjectUserSettings>()->bShowImportDialogAtReimport;
+	bShowOption = bShowDialogAtReimport && !bIsUnattended;
+
 	FText ReimportingText = FText::Format(LOCTEXT("AbcFactoryReimporting", "Reimporting {0}.abc"), FText::FromString(Obj->GetName()));
 	const FString& PageName = ReimportingText.ToString();
 	if (Obj->GetClass() == UStaticMesh::StaticClass())
@@ -582,6 +636,16 @@ EReimportResult::Type UAlembicImportFactory::ReimportGeometryCache(UGeometryCach
 			return EReimportResult::Cancelled;
 		}
 	}
+	else
+	{
+		UAbcImportSettings* ScriptedSettings = AssetImportTask
+			? Cast<UAbcImportSettings>(AssetImportTask->Options)
+			: nullptr;
+		if (ScriptedSettings)
+		{
+			ImportSettings = ScriptedSettings;
+		}
+	}
 
 	int32 NumThreads = 1;
 	if (FPlatformProcess::SupportsMultithreading())
@@ -661,6 +725,16 @@ EReimportResult::Type UAlembicImportFactory::ReimportSkeletalMesh(USkeletalMesh*
 			return EReimportResult::Cancelled;
 		}
 	}
+	else
+	{
+		UAbcImportSettings* ScriptedSettings = AssetImportTask
+			? Cast<UAbcImportSettings>(AssetImportTask->Options)
+			: nullptr;
+		if (ScriptedSettings)
+		{
+			ImportSettings = ScriptedSettings;
+		}
+	}
 
 	int32 NumThreads = 1;
 	if (FPlatformProcess::SupportsMultithreading())
@@ -781,6 +855,16 @@ EReimportResult::Type UAlembicImportFactory::ReimportStaticMesh(UStaticMesh* Mes
 			return EReimportResult::Cancelled;
 		}
 	}
+	else
+	{
+		UAbcImportSettings* ScriptedSettings = AssetImportTask
+			? Cast<UAbcImportSettings>(AssetImportTask->Options)
+			: nullptr;
+		if (ScriptedSettings)
+		{
+			ImportSettings = ScriptedSettings;
+		}
+	}
 
 	int32 NumThreads = 1;
 	if (FPlatformProcess::SupportsMultithreading())
@@ -797,6 +881,53 @@ EReimportResult::Type UAlembicImportFactory::ReimportStaticMesh(UStaticMesh* Mes
 	}
 	else
 	{
+		// Preserve settings on the existing mesh (same as on a FBX re-import)
+		const bool bOverrideExistingMaterials = false;
+		TSharedPtr<FExistingStaticMeshData> ExistingMeshData =
+			StaticMeshImportUtils::SaveExistingStaticMeshData(Mesh, bOverrideExistingMaterials, INDEX_NONE);
+
+		// Preserve the user data
+		const TArray<UAssetUserData*>* UserData = Mesh->GetAssetUserDataArray();
+		TMap<UAssetUserData*, bool> UserDataCopy;
+		if (UserData)
+		{
+			for (int32 Idx = 0; Idx < UserData->Num(); Idx++)
+			{
+				if ((*UserData)[Idx] != nullptr)
+				{
+					// This is slightly different from the equivalent in the FBX
+					// reimporter as it is using the deprecated version of the same
+					// method (StaticDuplicateObject).
+					FObjectDuplicationParameters DupParams(
+						(*UserData)[Idx], GetTransientPackage());
+					UAssetUserData* DupObject =
+						Cast<UAssetUserData>(StaticDuplicateObjectEx(DupParams));
+					bool bAddDupToRoot = !DupObject->IsRooted();
+					if (bAddDupToRoot)
+					{
+						DupObject->AddToRoot();
+					}
+					UserDataCopy.Add(DupObject, bAddDupToRoot);
+				}
+			}
+		}
+
+		// Preserve settings in navcollision subobject
+		UNavCollisionBase* NavCollision = nullptr;
+		if (Mesh->GetNavCollision())
+		{
+			FObjectDuplicationParameters DupParams(
+				Mesh->GetNavCollision(), GetTransientPackage());
+			NavCollision = Cast<UNavCollisionBase>(StaticDuplicateObjectEx(DupParams));
+		}
+
+		bool bAddedNavCollisionDupToRoot = false;
+		if (NavCollision && !NavCollision->IsRooted())
+		{
+			bAddedNavCollisionDupToRoot = true;
+			NavCollision->AddToRoot();
+		}
+
 		const TArray<UStaticMesh*>& StaticMeshes = Importer.ReimportAsStaticMesh(Mesh);
 		for (UStaticMesh* StaticMesh : StaticMeshes)
 		{
@@ -813,6 +944,43 @@ EReimportResult::Type UAlembicImportFactory::ReimportStaticMesh(UStaticMesh* Mes
 				{
 					Importer.UpdateAssetImportData(AssetImportData);
 				}
+
+				// Restore preserved user data
+				for (TPair<UAssetUserData*, bool> Kvp : UserDataCopy)
+				{
+					UAssetUserData* UserDataObject = Kvp.Key;
+					if (Kvp.Value)
+					{
+						// If the duplicated temporary UObject was added to root,
+						// we must remove it from the root
+						UserDataObject->RemoveFromRoot();
+					}
+					UserDataObject->Rename(nullptr, Mesh,
+						REN_DontCreateRedirectors | REN_DoNotDirty);
+					StaticMesh->AddAssetUserData(UserDataObject);
+				}
+
+				// Restore navcollision subobject
+				if (NavCollision)
+				{
+					if (bAddedNavCollisionDupToRoot)
+					{
+						// If the duplicated temporary UObject was added to root,
+						// we must remove it from the root
+						NavCollision->RemoveFromRoot();
+					}
+					StaticMesh->SetNavCollision(NavCollision);
+					NavCollision->Rename(NULL, Mesh,
+						REN_DontCreateRedirectors | REN_DoNotDirty);
+				}
+
+				// Restore additional mesh data settings
+				StaticMeshImportUtils::RestoreExistingMeshData(
+					ExistingMeshData,
+					StaticMesh,
+					INDEX_NONE,
+					false,
+					false);
 			}
 		}
 
diff --git a/Engine/Source/Editor/UnrealEd/Private/ImportUtils/StaticMeshImportUtils.cpp b/Engine/Source/Editor/UnrealEd/Private/ImportUtils/StaticMeshImportUtils.cpp
--- a/Engine/Source/Editor/UnrealEd/Private/ImportUtils/StaticMeshImportUtils.cpp
+++ b/Engine/Source/Editor/UnrealEd/Private/ImportUtils/StaticMeshImportUtils.cpp
@@ -585,6 +585,12 @@ bool StaticMeshImportUtils::AddConvexGeomFromVertices( const TArray<FVector3f>&
 	return true;
 }
 
+TSharedPtr<FExistingStaticMeshData> StaticMeshImportUtils::SaveExistingStaticMeshData(UStaticMesh* ExistingMesh, bool bImportMaterials, int32 LodIndex)
+{
+	UnFbx::FBXImportOptions ImportOptions;
+	ImportOptions.bImportMaterials = bImportMaterials;
+	return SaveExistingStaticMeshData(ExistingMesh, &ImportOptions, LodIndex);
+}
 
 TSharedPtr<FExistingStaticMeshData> StaticMeshImportUtils::SaveExistingStaticMeshData(UStaticMesh* ExistingMesh, UnFbx::FBXImportOptions* ImportOptions, int32 LodIndex)
 {
@@ -737,6 +743,13 @@ TSharedPtr<FExistingStaticMeshData> StaticMeshImportUtils::SaveExistingStaticMes
 	ExistingMeshDataPtr->ExistingPositiveBoundsExtension = (FVector3f)ExistingMesh->GetPositiveBoundsExtension();
 	ExistingMeshDataPtr->ExistingNegativeBoundsExtension = (FVector3f)ExistingMesh->GetNegativeBoundsExtension();
 
+	ExistingMeshDataPtr->ExistingSupportPhysicalMaterialMasks = ExistingMesh->bSupportPhysicalMaterialMasks;
+	ExistingMeshDataPtr->ExistingSupportGpuUniformlyDistributedSampling = ExistingMesh->bSupportGpuUniformlyDistributedSampling;
+	ExistingMeshDataPtr->ExistingSupportRayTracing = ExistingMesh->bSupportRayTracing;
+	ExistingMeshDataPtr->ExistingForceMiplevelsToBeResident = ExistingMesh->bGlobalForceMipLevelsToBeResident;
+	ExistingMeshDataPtr->ExistingNeverStream = ExistingMesh->NeverStream;
+	ExistingMeshDataPtr->ExistingNumCinematicMipLevels = ExistingMesh->NumCinematicMipLevels;
+
 	UFbxStaticMeshImportData* ImportData = Cast<UFbxStaticMeshImportData>(ExistingMesh->AssetImportData);
 	if (ImportData && ExistingMeshDataPtr->UseMaterialNameSlotWorkflow)
 	{
@@ -1266,6 +1279,13 @@ void StaticMeshImportUtils::RestoreExistingMeshData(const TSharedPtr<const FExis
 	NewMesh->SetNegativeBoundsExtension((FVector)ExistingMeshDataPtr->ExistingNegativeBoundsExtension);
 
 	NewMesh->ComplexCollisionMesh = ExistingMeshDataPtr->ExistingComplexCollisionMesh;
+
+	NewMesh->bSupportPhysicalMaterialMasks = ExistingMeshDataPtr->ExistingSupportPhysicalMaterialMasks;
+	NewMesh->bSupportGpuUniformlyDistributedSampling = ExistingMeshDataPtr->ExistingSupportGpuUniformlyDistributedSampling;
+	NewMesh->bSupportRayTracing = ExistingMeshDataPtr->ExistingSupportRayTracing;
+	NewMesh->bGlobalForceMipLevelsToBeResident = ExistingMeshDataPtr->ExistingForceMiplevelsToBeResident;
+	NewMesh->NeverStream = ExistingMeshDataPtr->ExistingNeverStream;
+	NewMesh->NumCinematicMipLevels = ExistingMeshDataPtr->ExistingNumCinematicMipLevels;
 }
 
 #undef LOCTEXT_NAMESPACE
diff --git a/Engine/Source/Editor/UnrealEd/Public/ImportUtils/StaticMeshImportUtils.h b/Engine/Source/Editor/UnrealEd/Public/ImportUtils/StaticMeshImportUtils.h
--- a/Engine/Source/Editor/UnrealEd/Public/ImportUtils/StaticMeshImportUtils.h
+++ b/Engine/Source/Editor/UnrealEd/Public/ImportUtils/StaticMeshImportUtils.h
@@ -81,6 +81,14 @@ struct FExistingStaticMeshData
 	FVector3f					ExistingPositiveBoundsExtension;
 	FVector3f					ExistingNegativeBoundsExtension;
 
+	bool						ExistingSupportPhysicalMaterialMasks;
+	bool						ExistingSupportGpuUniformlyDistributedSampling;
+	bool						ExistingSupportRayTracing;
+	int32						ExistingLODForOccluderMesh;
+	bool						ExistingForceMiplevelsToBeResident;
+	bool						ExistingNeverStream;
+	int32						ExistingNumCinematicMipLevels;
+
 	UStaticMesh::FOnMeshChanged	ExistingOnMeshChanged;
 	UStaticMesh* ExistingComplexCollisionMesh = nullptr;
 
@@ -120,6 +128,8 @@ namespace StaticMeshImportUtils
 	/** Utility for adding one convex hull from the given verts */
 	UNREALED_API bool AddConvexGeomFromVertices(const TArray<FVector3f>& Verts, FKAggregateGeom* AggGeom, const TCHAR* ObjName);
 
+	UNREALED_API TSharedPtr<FExistingStaticMeshData> SaveExistingStaticMeshData(UStaticMesh* ExistingMesh, bool bImportMaterials, int32 LodIndex);
+
 	UNREALED_API TSharedPtr<FExistingStaticMeshData> SaveExistingStaticMeshData(UStaticMesh* ExistingMesh, UnFbx::FBXImportOptions* ImportOptions, int32 LodIndex);
 
 	/* This function is call before building the mesh when we do a re-import*/
