Alembic: Fixed GeometryCache import with "Flatten Tracks" disabled had bad time offset computation.

diff --git a/Engine/Plugins/Importers/AlembicImporter/Source/AlembicLibrary/Private/AbcImporter.cpp b/Engine/Plugins/Importers/AlembicImporter/Source/AlembicLibrary/Private/AbcImporter.cpp
--- a/Engine/Plugins/Importers/AlembicImporter/Source/AlembicLibrary/Private/AbcImporter.cpp
+++ b/Engine/Plugins/Importers/AlembicImporter/Source/AlembicLibrary/Private/AbcImporter.cpp
@@ -488,7 +488,7 @@ UGeometryCache* FAbcImporter::ImportAsGeometryCache(UObject* InParent, EObjectFl
 
 							// Generate the mesh data for this sample
 							const bool bVisible = PolyMesh->GetVisibility(FrameIndex);
-							const float FrameTime = PolyMesh->GetTimeForFrameIndex(FrameIndex);
+							const float FrameTime = PolyMesh->GetTimeForFrameIndex(FrameIndex)- InAbcFile->GetImportTimeOffset();
 							if (bVisible)
 							{
 								const bool bUseVelocitiesAsMotionVectors = ( ImportSettings->GeometryCacheSettings.MotionVectors == EAbcGeometryCacheMotionVectorsImport::ImportAbcVelocitiesAsMotionVectors );
