Alembic: Added a check when importing Alembic as skeletal mesh to detect animations that are too big to be imported due to memory constraint.

diff --git a/Engine/Plugins/Importers/AlembicImporter/Source/AlembicLibrary/Private/AbcImporter.cpp b/Engine/Plugins/Importers/AlembicImporter/Source/AlembicLibrary/Private/AbcImporter.cpp
--- a/Engine/Plugins/Importers/AlembicImporter/Source/AlembicLibrary/Private/AbcImporter.cpp
+++ b/Engine/Plugins/Importers/AlembicImporter/Source/AlembicLibrary/Private/AbcImporter.cpp
@@ -881,7 +881,6 @@ const bool FAbcImporter::CompressAnimationDataUsingPCA(const FAbcCompressionSett
 		if (InCompressionSettings.bMergeMeshes)
 		{
 			// Merged path
-			const uint32 FrameZeroIndex = 0;
 			TArray<FVector3f> AverageVertexData;
 			TArray<FVector3f> AverageNormalData;
 
@@ -932,6 +931,39 @@ const bool FAbcImporter::CompressAnimationDataUsingPCA(const FAbcCompressionSett
 				Flags |= EFrameReadFlags::ApplyMatrix;
 			}
 
+			{
+				// Check the first frame to see if the Alembic can be imported as a skeletal mesh due to memory constraints
+				AbcFile->ReadFrame(AbcFile->GetStartFrameIndex(), Flags, 0);
+				MergedMeshesFunc(AbcFile->GetStartFrameIndex(), AbcFile);
+				AbcFile->CleanupFrameData(0);
+
+				const int32 NumFrames = AbcFile->GetEndFrameIndex() - AbcFile->GetStartFrameIndex() + 1;
+				const uint64 NumMatrixElements = uint64(AverageVertexData.Num()) * 3 * NumFrames;
+				if (!IntFitsIn<uint32>(NumMatrixElements))
+				{
+					UE_LOG(LogAbcImporter, Error, TEXT("Vertex matrix has too many elements (%llu) because the mesh has too many vertices (%d) and/or the animation has too many frames (%d). Try importing as GeometryCache instead."),
+						NumMatrixElements, AverageVertexData.Num(), NumFrames);
+					return false;
+				}
+
+				const uint64 NumNormalsMatrixElements = uint64(AverageNormalData.Num()) * 3 * NumFrames;
+				if (!IntFitsIn<uint32>(NumNormalsMatrixElements))
+				{
+					UE_LOG(LogAbcImporter, Error, TEXT("Normal matrix has too many elements (%llu) because the mesh has too many vertices (%d) and/or the animation has too many frames (%d). Try importing as GeometryCache instead."),
+						NumNormalsMatrixElements, AverageNormalData.Num(), NumFrames);
+					return false;
+				}
+
+				AverageVertexData.Reset();
+				AverageNormalData.Reset();
+				ObjectVertexOffsets.Reset();
+				ObjectIndexOffsets.Reset();
+
+				MinTime = FLT_MAX;
+				MaxTime = -FLT_MAX;
+				NumSamples = 0;
+			}
+
 			AbcFile->ProcessFrames(MergedMeshesFunc, Flags);
 
 			// Average out vertex data
@@ -1151,6 +1183,45 @@ const bool FAbcImporter::CompressAnimationDataUsingPCA(const FAbcCompressionSett
 				Flags |= EFrameReadFlags::ApplyMatrix;
 			}
 
+			{
+				// Check the first frame to see if the Alembic can be imported as a skeletal mesh due to memory constraints
+				AbcFile->ReadFrame(AbcFile->GetStartFrameIndex(), Flags, 0);
+				IndividualMeshesFunc(AbcFile->GetStartFrameIndex(), AbcFile);
+				AbcFile->CleanupFrameData(0);
+
+				const int32 NumFrames = AbcFile->GetEndFrameIndex() - AbcFile->GetStartFrameIndex() + 1;
+				for (int32 MeshIndex = 0; MeshIndex < NumPolyMeshesToCompress; ++MeshIndex)
+				{
+					const uint64 NumMatrixElements = uint64(AverageVertexData[MeshIndex].Num()) * 3 * NumFrames;
+					if (!IntFitsIn<uint32>(NumMatrixElements))
+					{
+						UE_LOG(LogAbcImporter, Error, TEXT("Vertex matrix has too many elements (%llu) because the mesh has too many vertices (%d) and/or the animation has too many frames (%d). Try importing as GeometryCache instead."),
+						NumMatrixElements, AverageVertexData[MeshIndex].Num(), NumFrames);
+						return false;
+					}
+
+					const uint64 NumNormalsMatrixElements = uint64(AverageNormalData[MeshIndex].Num()) * 3 * NumFrames;
+					if (!IntFitsIn<uint32>(NumNormalsMatrixElements))
+					{
+						UE_LOG(LogAbcImporter, Error, TEXT("Normal matrix has too many elements (%llu) because the mesh has too many vertices (%d) and/or the animation has too many frames (%d). Try importing as GeometryCache instead."),
+						NumNormalsMatrixElements, AverageNormalData[MeshIndex].Num(), NumFrames);
+						return false;
+					}
+				}
+
+				MinTimes.Reset();
+				MaxTimes.Reset();
+				AverageVertexData.Reset();
+				AverageNormalData.Reset();
+
+				MinTimes.AddZeroed(NumPolyMeshesToCompress);
+				MaxTimes.AddZeroed(NumPolyMeshesToCompress);
+				AverageVertexData.AddDefaulted(NumPolyMeshesToCompress);
+				AverageNormalData.AddDefaulted(NumPolyMeshesToCompress);
+
+				NumSamples = 0;
+			}
+
 			AbcFile->ProcessFrames(IndividualMeshesFunc, Flags);
 
 			// Average out vertex data
@@ -1166,7 +1237,6 @@ const bool FAbcImporter::CompressAnimationDataUsingPCA(const FAbcCompressionSett
 			}
 			const FVector AverageSampleCenter = AverageBoundingBox.GetCenter();
 
-
 			TArray<TArray<float>> Matrices;
 			TArray<TArray<float>> NormalsMatrices;
 			for (int32 MeshIndex = 0; MeshIndex < NumPolyMeshesToCompress; ++MeshIndex)
