Alembic: Added a check when importing Alembic as skeletal mesh to detect animations that are too big to be imported due to memory constraint.

diff --git a/Engine/Plugins/Importers/AlembicImporter/Source/AlembicLibrary/Private/AbcImporter.cpp b/Engine/Plugins/Importers/AlembicImporter/Source/AlembicLibrary/Private/AbcImporter.cpp
--- a/Engine/Plugins/Importers/AlembicImporter/Source/AlembicLibrary/Private/AbcImporter.cpp
+++ b/Engine/Plugins/Importers/AlembicImporter/Source/AlembicLibrary/Private/AbcImporter.cpp
@@ -939,7 +939,7 @@ const bool FAbcImporter::CompressAnimationDataUsingPCA(const FAbcCompressionSett
 
 				const int32 NumFrames = AbcFile->GetEndFrameIndex() - AbcFile->GetStartFrameIndex() + 1;
 				const uint64 NumMatrixElements = uint64(AverageVertexData.Num()) * 3 * NumFrames;
-				if (!IntFitsIn<uint32>(NumMatrixElements))
+				if (!IntFitsIn<int32>(NumMatrixElements))
 				{
 					UE_LOG(LogAbcImporter, Error, TEXT("Vertex matrix has too many elements (%llu) because the mesh has too many vertices (%d) and/or the animation has too many frames (%d). Try importing as GeometryCache instead."),
 						NumMatrixElements, AverageVertexData.Num(), NumFrames);
@@ -947,7 +947,7 @@ const bool FAbcImporter::CompressAnimationDataUsingPCA(const FAbcCompressionSett
 				}
 
 				const uint64 NumNormalsMatrixElements = uint64(AverageNormalData.Num()) * 3 * NumFrames;
-				if (!IntFitsIn<uint32>(NumNormalsMatrixElements))
+				if (!IntFitsIn<int32>(NumNormalsMatrixElements))
 				{
 					UE_LOG(LogAbcImporter, Error, TEXT("Normal matrix has too many elements (%llu) because the mesh has too many vertices (%d) and/or the animation has too many frames (%d). Try importing as GeometryCache instead."),
 						NumNormalsMatrixElements, AverageNormalData.Num(), NumFrames);
@@ -1193,7 +1193,7 @@ const bool FAbcImporter::CompressAnimationDataUsingPCA(const FAbcCompressionSett
 				for (int32 MeshIndex = 0; MeshIndex < NumPolyMeshesToCompress; ++MeshIndex)
 				{
 					const uint64 NumMatrixElements = uint64(AverageVertexData[MeshIndex].Num()) * 3 * NumFrames;
-					if (!IntFitsIn<uint32>(NumMatrixElements))
+					if (!IntFitsIn<int32>(NumMatrixElements))
 					{
 						UE_LOG(LogAbcImporter, Error, TEXT("Vertex matrix has too many elements (%llu) because the mesh has too many vertices (%d) and/or the animation has too many frames (%d). Try importing as GeometryCache instead."),
 						NumMatrixElements, AverageVertexData[MeshIndex].Num(), NumFrames);
@@ -1201,7 +1201,7 @@ const bool FAbcImporter::CompressAnimationDataUsingPCA(const FAbcCompressionSett
 					}
 
 					const uint64 NumNormalsMatrixElements = uint64(AverageNormalData[MeshIndex].Num()) * 3 * NumFrames;
-					if (!IntFitsIn<uint32>(NumNormalsMatrixElements))
+					if (!IntFitsIn<int32>(NumNormalsMatrixElements))
 					{
 						UE_LOG(LogAbcImporter, Error, TEXT("Normal matrix has too many elements (%llu) because the mesh has too many vertices (%d) and/or the animation has too many frames (%d). Try importing as GeometryCache instead."),
 						NumNormalsMatrixElements, AverageNormalData[MeshIndex].Num(), NumFrames);
