Alembic: Fixed import options window missing when re-importing an asset from Alembic.

diff --git a/Engine/Plugins/Importers/AlembicImporter/Source/AlembicImporter/Private/AlembicImportFactory.cpp b/Engine/Plugins/Importers/AlembicImporter/Source/AlembicImporter/Private/AlembicImportFactory.cpp
--- a/Engine/Plugins/Importers/AlembicImporter/Source/AlembicImporter/Private/AlembicImportFactory.cpp
+++ b/Engine/Plugins/Importers/AlembicImporter/Source/AlembicImporter/Private/AlembicImportFactory.cpp
@@ -78,10 +78,7 @@ UObject* UAlembicImportFactory::FactoryCreateFile(UClass* InClass, UObject* InPa
 		|| FApp::IsUnattended()
 		|| IsRunningCommandlet()
 		|| GIsRunningUnattendedScript);
-	if (bIsUnattended)
-	{
-		bShowOption = false;
-	}
+	bShowOption = !bIsUnattended;
 
 	// Check if it's a re-import
 	if (InParent != nullptr)
@@ -445,9 +442,7 @@ EReimportResult::Type UAlembicImportFactory::Reimport(UObject* Obj)
 		|| FApp::IsUnattended()
 		|| IsRunningCommandlet()
 		|| GIsRunningUnattendedScript);
-	const bool bShowDialogAtReimport =
-		GetDefault<UEditorPerProjectUserSettings>()->bShowImportDialogAtReimport;
-	bShowOption = bShowDialogAtReimport && !bIsUnattended;
+	bShowOption = !bIsUnattended;
 
 	FText ReimportingText = FText::Format(LOCTEXT("AbcFactoryReimporting", "Reimporting {0}.abc"), FText::FromString(Obj->GetName()));
 	const FString& PageName = ReimportingText.ToString();
