Alembic: Fixed transform chain not being composed properly when a transform in the hierarchy is detected as constant.

diff --git a/Engine/Plugins/Importers/AlembicImporter/Source/AlembicLibrary/Private/AbcTransform.cpp b/Engine/Plugins/Importers/AlembicImporter/Source/AlembicLibrary/Private/AbcTransform.cpp
--- a/Engine/Plugins/Importers/AlembicImporter/Source/AlembicLibrary/Private/AbcTransform.cpp
+++ b/Engine/Plugins/Importers/AlembicImporter/Source/AlembicLibrary/Private/AbcTransform.cpp
@@ -59,7 +59,7 @@ FMatrix FAbcTransform::GetMatrix(const int32 FrameIndex) const
 {
 	if (bConstantIdentity || bConstant)
 	{
-		return InitialValue;
+		return Parent ? InitialValue * Parent->GetMatrix(FrameIndex) : InitialValue;
 	}
 
 	// Find matrix within resident samples
@@ -71,7 +71,7 @@ FMatrix FAbcTransform::GetMatrix(const int32 FrameIndex) const
 		}
 	}
 
-	return InitialValue;
+	return Parent ? InitialValue * Parent->GetMatrix(FrameIndex) : InitialValue;
 }
 
 bool FAbcTransform::HasConstantTransform() const
