Alembic:
* Added material slot names to Geometry Cache and related functions on the component. In particular,calling SetMaterialByName on the component in Blueprint will now work.
* Face set names are now imported from Alembic as material slot names for GeometryCache (and static and skeletal meshes).

diff --git a/Engine/Plugins/Importers/AlembicImporter/Source/AlembicLibrary/Private/AbcImporter.cpp b/Engine/Plugins/Importers/AlembicImporter/Source/AlembicLibrary/Private/AbcImporter.cpp
--- a/Engine/Plugins/Importers/AlembicImporter/Source/AlembicLibrary/Private/AbcImporter.cpp
+++ b/Engine/Plugins/Importers/AlembicImporter/Source/AlembicLibrary/Private/AbcImporter.cpp
@@ -77,6 +77,9 @@ DEFINE_LOG_CATEGORY_STATIC(LogAbcImporter, Verbose, All);
 
 #define PRINT_UNIQUE_VERTICES 0
 
+static const FString NoFaceSetNameStr(TEXT("NoFaceSetName"));
+static const FName NoFaceSetName(TEXT("NoFaceSetName"));
+
 FAbcImporter::FAbcImporter()
 	: ImportSettings(nullptr), AbcFile(nullptr)
 {
@@ -238,7 +241,14 @@ UStaticMesh* FAbcImporter::CreateStaticMeshFromSample(UObject* InParent, const F
 				}
 			}
 
-			StaticMesh->GetStaticMaterials().Add((Material != nullptr) ? Material : DefaultMaterial);
+			FStaticMaterial StaticMaterial;
+			StaticMaterial.MaterialInterface = (Material != nullptr) ? Material : DefaultMaterial;
+
+			FName SlotName((Material != nullptr) ? FName(FaceSetNames[MaterialIndex]) : NoFaceSetName);
+			StaticMaterial.MaterialSlotName = SlotName;
+			StaticMaterial.ImportedMaterialSlotName = SlotName;
+
+			StaticMesh->GetStaticMaterials().Add(StaticMaterial);
 		}
 
 		GenerateMeshDescriptionFromSample(Sample, MeshDescription, StaticMesh);
@@ -307,8 +317,7 @@ const TArray<UStaticMesh*> FAbcImporter::ImportAsStaticMesh(UObject* InParent, E
 					else
 					{
 						// Default name
-						static const FString DefaultName("NoFaceSetName");
-						MergedFaceSetNames.Add(DefaultName);
+						MergedFaceSetNames.Add(NoFaceSetNameStr);
 					}
 				}
 			}
@@ -442,6 +451,7 @@ UGeometryCache* FAbcImporter::ImportAsGeometryCache(UObject* InParent, EObjectFl
 				{
 					UMaterialInterface* Material = AbcImporterUtilities::RetrieveMaterial(*AbcFile, FaceSetName, InParent, Flags);
 					GeometryCache->Materials.Add((Material != nullptr) ? Material : DefaultMaterial);		
+					GeometryCache->MaterialSlotNames.Add(FaceSetName != TEXT("DefaultMaterial") ? FName(FaceSetName) : NoFaceSetName);
 
 					if (Material != UMaterial::GetDefaultMaterial(MD_Surface))
 					{
@@ -480,11 +490,16 @@ UGeometryCache* FAbcImporter::ImportAsGeometryCache(UObject* InParent, EObjectFl
 							if (PolyMesh->FaceSetNames.IsValidIndex(MaterialIndex))
 							{
 								Material = AbcImporterUtilities::RetrieveMaterial(*AbcFile, PolyMesh->FaceSetNames[MaterialIndex], InParent, Flags);
+								GeometryCache->MaterialSlotNames.Add(FName(PolyMesh->FaceSetNames[MaterialIndex]));
 								if (Material != UMaterial::GetDefaultMaterial(MD_Surface))
 								{
 									Material->PostEditChange();
 								}
 							}
+							else
+							{
+								GeometryCache->MaterialSlotNames.Add(NoFaceSetName);
+							}
 
 							GeometryCache->Materials.Add((Material != nullptr) ? Material : DefaultMaterial);
 						}
@@ -744,7 +759,7 @@ TArray<UObject*> FAbcImporter::ImportAsSkeletalMesh(UObject* InParent, EObjectFl
 				{
 					const FString& MaterialName = CompressedData.MaterialNames[MaterialIndex];
 					UMaterialInterface* Material = AbcImporterUtilities::RetrieveMaterial(*AbcFile, MaterialName, InParent, Flags);
-					SkeletalMesh->GetMaterials().Add(FSkeletalMaterial(Material, true));
+					SkeletalMesh->GetMaterials().Add(FSkeletalMaterial(Material, true, false, FName(MaterialName), FName(MaterialName)));
 					if (Material != UMaterial::GetDefaultMaterial(MD_Surface))
 					{
 						Material->PostEditChange();
@@ -1014,8 +1029,7 @@ const bool FAbcImporter::CompressAnimationDataUsingPCA(const FAbcCompressionSett
 				}
 				else
 				{
-					static const FString DefaultName("NoFaceSetName");
-					CompressedData.MaterialNames.Add(DefaultName);
+					CompressedData.MaterialNames.Add(NoFaceSetNameStr);
 				}
 			}
 
@@ -1386,8 +1400,7 @@ const bool FAbcImporter::CompressAnimationDataUsingPCA(const FAbcCompressionSett
 				}
 				else
 				{
-					static const FString DefaultName("NoFaceSetName");
-					CompressedData.MaterialNames.Add(DefaultName);
+					CompressedData.MaterialNames.Add(NoFaceSetNameStr);
 				}
 
 				if (bRunComparison)
@@ -1429,8 +1442,7 @@ const bool FAbcImporter::CompressAnimationDataUsingPCA(const FAbcCompressionSett
 		}
 		else
 		{
-			static const FString DefaultName("NoFaceSetName");
-			CompressedData.MaterialNames.Add(DefaultName);
+			CompressedData.MaterialNames.Add(NoFaceSetNameStr);
 		}
 	}
 		
diff --git a/Engine/Plugins/Runtime/GeometryCache/Source/GeometryCache/Classes/GeometryCache.h b/Engine/Plugins/Runtime/GeometryCache/Source/GeometryCache/Classes/GeometryCache.h
--- a/Engine/Plugins/Runtime/GeometryCache/Source/GeometryCache/Classes/GeometryCache.h
+++ b/Engine/Plugins/Runtime/GeometryCache/Source/GeometryCache/Classes/GeometryCache.h
@@ -54,9 +54,12 @@ public:
 	UPROPERTY(VisibleAnywhere, Instanced, Category = Thumbnail)
 	TObjectPtr<class UThumbnailInfo> ThumbnailInfo;
 #endif
-	UPROPERTY(EditAnywhere, Category = GeometryCache)
+	UPROPERTY(EditAnywhere, BlueprintReadWrite, Category = GeometryCache)
 	TArray<TObjectPtr<UMaterialInterface>> Materials;
 	
+	UPROPERTY(EditAnywhere, BlueprintReadWrite, Category = GeometryCache)
+	TArray<FName> MaterialSlotNames;
+
 	/** GeometryCache track defining the samples/geometry data for this GeomCache instance */
 	UPROPERTY(VisibleAnywhere, Category=GeometryCache)
 	TArray<TObjectPtr<UGeometryCacheTrack>> Tracks;
diff --git a/Engine/Plugins/Runtime/GeometryCache/Source/GeometryCache/Classes/GeometryCacheComponent.h b/Engine/Plugins/Runtime/GeometryCache/Source/GeometryCache/Classes/GeometryCacheComponent.h
--- a/Engine/Plugins/Runtime/GeometryCache/Source/GeometryCache/Classes/GeometryCacheComponent.h
+++ b/Engine/Plugins/Runtime/GeometryCache/Source/GeometryCache/Classes/GeometryCacheComponent.h
@@ -69,6 +69,9 @@ class GEOMETRYCACHE_API UGeometryCacheComponent : public UMeshComponent
 	//~ Begin UMeshComponent Interface.
 	virtual int32 GetNumMaterials() const override;
 	virtual UMaterialInterface* GetMaterial(int32 MaterialIndex) const override;
+	virtual int32 GetMaterialIndex(FName MaterialSlotName) const override;
+	virtual TArray<FName> GetMaterialSlotNames() const override;
+	virtual bool IsMaterialSlotNameValid(FName MaterialSlotName) const override;
 	//~ End UMeshComponent Interface.
 
 
diff --git a/Engine/Plugins/Runtime/GeometryCache/Source/GeometryCache/Private/GeometryCacheComponent.cpp b/Engine/Plugins/Runtime/GeometryCache/Source/GeometryCache/Private/GeometryCacheComponent.cpp
--- a/Engine/Plugins/Runtime/GeometryCache/Source/GeometryCache/Private/GeometryCacheComponent.cpp
+++ b/Engine/Plugins/Runtime/GeometryCache/Source/GeometryCache/Private/GeometryCacheComponent.cpp
@@ -261,6 +261,21 @@ UMaterialInterface* UGeometryCacheComponent::GetMaterial(int32 MaterialIndex) co
 	}
 }
 
+int32 UGeometryCacheComponent::GetMaterialIndex(FName MaterialSlotName) const
+{
+	return GeometryCache ? GeometryCache->MaterialSlotNames.IndexOfByKey(MaterialSlotName) : INDEX_NONE;
+}
+
+TArray<FName> UGeometryCacheComponent::GetMaterialSlotNames() const
+{
+	return GeometryCache ? GeometryCache->MaterialSlotNames : TArray<FName>();
+}
+
+bool UGeometryCacheComponent::IsMaterialSlotNameValid(FName MaterialSlotName) const
+{
+	return GetMaterialIndex(MaterialSlotName) != INDEX_NONE;
+}
+
 void UGeometryCacheComponent::CreateTrackSection(int32 TrackIndex)
 {
 	// Ensure sections array is long enough
