Truncation fixes for AlembicLibrary

diff --git a/Engine/Plugins/Importers/AlembicImporter/Source/AlembicLibrary/AlembicLibrary.Build.cs b/Engine/Plugins/Importers/AlembicImporter/Source/AlembicLibrary/AlembicLibrary.Build.cs
--- a/Engine/Plugins/Importers/AlembicImporter/Source/AlembicLibrary/AlembicLibrary.Build.cs
+++ b/Engine/Plugins/Importers/AlembicImporter/Source/AlembicLibrary/AlembicLibrary.Build.cs
@@ -39,5 +39,6 @@ public class AlembicLibrary : ModuleRules
             }
         );
 
-    }
+		UnsafeTypeCastWarningLevel = WarningLevel.Error;
+	}
 }
diff --git a/Engine/Plugins/Importers/AlembicImporter/Source/AlembicLibrary/Private/AbcFile.cpp b/Engine/Plugins/Importers/AlembicImporter/Source/AlembicLibrary/Private/AbcFile.cpp
--- a/Engine/Plugins/Importers/AlembicImporter/Source/AlembicLibrary/Private/AbcFile.cpp
+++ b/Engine/Plugins/Importers/AlembicImporter/Source/AlembicLibrary/Private/AbcFile.cpp
@@ -154,7 +154,7 @@ EAbcImportError FAbcFile::Open()
 	Alembic::Abc::TimeSamplingPtr TimeSampler = Archive.getTimeSampling(TimeSamplingIndex);
 	if (TimeSampler)
 	{
-		ArchiveSecondsPerFrame = TimeSampler->getTimeSamplingType().getTimePerCycle();
+		ArchiveSecondsPerFrame = static_cast<float>(TimeSampler->getTimeSamplingType().getTimePerCycle());
 	}
 
 	MeshUtilities = FModuleManager::Get().LoadModulePtr<IMeshUtilities>("MeshUtilities");
@@ -254,24 +254,24 @@ EAbcImportError FAbcFile::Import(UAbcImportSettings* InImportSettings)
 
 	SecondsPerFrame = TimeStep;
 	FramesPerSecond = TimeStep > 0.f ? FMath::RoundToInt(1.f / TimeStep) : 30;
-	ImportLength = FrameSpan * TimeStep;
+	ImportLength = static_cast<float>(FrameSpan) * TimeStep;
 
 	// Calculate time offset from start of import animation range
-	ImportTimeOffset = StartFrameIndex * SecondsPerFrame;
+	ImportTimeOffset = static_cast<float>(StartFrameIndex) * SecondsPerFrame;
 
 	// Read first-frames for both the transforms and poly meshes
 
 	bool bValidFirstFrames = true;
 	for (FAbcTransform* Transform : Transforms)
 	{
-		bValidFirstFrames &= Transform->ReadFirstFrame(StartFrameIndex * SecondsPerFrame, StartFrameIndex);
+		bValidFirstFrames &= Transform->ReadFirstFrame(static_cast<float>(StartFrameIndex) * SecondsPerFrame, StartFrameIndex);
 	}
 
 	for (FAbcPolyMesh* PolyMesh : PolyMeshes)
 	{
 		if (PolyMesh->bShouldImport)
 		{
-			bValidFirstFrames &= PolyMesh->ReadFirstFrame(StartFrameIndex * SecondsPerFrame, StartFrameIndex);
+			bValidFirstFrames &= PolyMesh->ReadFirstFrame(static_cast<float>(StartFrameIndex) * SecondsPerFrame, StartFrameIndex);
 		}
 	}	
 
@@ -501,7 +501,7 @@ void FAbcFile::ReadFrame(int32 FrameIndex, const EFrameReadFlags InFlags, const
 {
 	for (IAbcObject* Object : Objects)
 	{
-		Object->SetFrameAndTime(FrameIndex * SecondsPerFrame, FrameIndex, InFlags, ReadIndex);
+		Object->SetFrameAndTime(static_cast<float>(FrameIndex) * SecondsPerFrame, FrameIndex, InFlags, ReadIndex);
 	}
 }
 
diff --git a/Engine/Plugins/Importers/AlembicImporter/Source/AlembicLibrary/Private/AbcImportUtilities.cpp b/Engine/Plugins/Importers/AlembicImporter/Source/AlembicLibrary/Private/AbcImportUtilities.cpp
--- a/Engine/Plugins/Importers/AlembicImporter/Source/AlembicLibrary/Private/AbcImportUtilities.cpp
+++ b/Engine/Plugins/Importers/AlembicImporter/Source/AlembicLibrary/Private/AbcImportUtilities.cpp
@@ -1202,9 +1202,9 @@ void AbcImporterUtilities::GenerateDeltaFrameDataMatrix(const TArray<FVector3f>&
 		const int32 ComponentIndexOffset = (VertexIndex + AverageVertexOffset) * 3;
 		const FVector AverageDifference = ((FVector)AverageVertexData[VertexIndex + AverageVertexOffset] + SamplePositionOffset) - (FVector)FrameVertexData[VertexIndex];
 
-		OutGeneratedMatrix[VertexOffset + ComponentIndexOffset + 0] = AverageDifference.X;
-		OutGeneratedMatrix[VertexOffset + ComponentIndexOffset + 1] = AverageDifference.Y;
-		OutGeneratedMatrix[VertexOffset + ComponentIndexOffset + 2] = AverageDifference.Z;
+		OutGeneratedMatrix[VertexOffset + ComponentIndexOffset + 0] = static_cast<float>(AverageDifference.X);
+		OutGeneratedMatrix[VertexOffset + ComponentIndexOffset + 1] = static_cast<float>(AverageDifference.Y);
+		OutGeneratedMatrix[VertexOffset + ComponentIndexOffset + 2] = static_cast<float>(AverageDifference.Z);
 	}
 
 	const uint32 NumIndices = FrameNormalData.Num();
@@ -1215,9 +1215,9 @@ void AbcImporterUtilities::GenerateDeltaFrameDataMatrix(const TArray<FVector3f>&
 		const int32 ComponentIndexOffset = (Index + AverageIndexOffset) * 3;
 		const FVector AverageNormal = (FVector)AverageNormalData[Index + AverageIndexOffset] - (FVector)FrameNormalData[Index];
 
-		OutGeneratedNormalsMatrix[IndexOffset + ComponentIndexOffset + 0] = AverageNormal.X;
-		OutGeneratedNormalsMatrix[IndexOffset + ComponentIndexOffset + 1] = AverageNormal.Y;
-		OutGeneratedNormalsMatrix[IndexOffset + ComponentIndexOffset + 2] = AverageNormal.Z;
+		OutGeneratedNormalsMatrix[IndexOffset + ComponentIndexOffset + 0] = static_cast<float>(AverageNormal.X);
+		OutGeneratedNormalsMatrix[IndexOffset + ComponentIndexOffset + 1] = static_cast<float>(AverageNormal.Y);
+		OutGeneratedNormalsMatrix[IndexOffset + ComponentIndexOffset + 2] = static_cast<float>(AverageNormal.Z);
 	}
 }
 
@@ -1274,7 +1274,7 @@ void AbcImporterUtilities::GenerateCompressedMeshData(FCompressedAbcData& Compre
 		for (uint32 CurveSampleIndex = 0; CurveSampleIndex < NumSamples; ++CurveSampleIndex)
 		{
 			CurveValues.Add(BasesWeights[BaseIndex + (OriginalNumberOfSingularValues * CurveSampleIndex)]);
-			TimeValues.Add(StartTime + (SampleTimeStep * CurveSampleIndex));
+			TimeValues.Add(StartTime + (SampleTimeStep * static_cast<float>(CurveSampleIndex)));
 		}
 	}
 }
@@ -1282,8 +1282,8 @@ void AbcImporterUtilities::GenerateCompressedMeshData(FCompressedAbcData& Compre
 void AbcImporterUtilities::CalculateNewStartAndEndFrameIndices(const float FrameStepRatio, int32& InOutStartFrameIndex, int32& InOutEndFrameIndex )
 {
 	// Using the calculated ratio we recompute the start/end frame indices
-	InOutStartFrameIndex = FMath::Max(FMath::FloorToInt(InOutStartFrameIndex * FrameStepRatio), 0);
-	InOutEndFrameIndex = FMath::CeilToInt(InOutEndFrameIndex * FrameStepRatio);
+	InOutStartFrameIndex = FMath::Max(FMath::FloorToInt(static_cast<float>(InOutStartFrameIndex) * FrameStepRatio), 0);
+	InOutEndFrameIndex = FMath::CeilToInt(static_cast<float>(InOutEndFrameIndex) * FrameStepRatio);
 }
 
 bool AbcImporterUtilities::AreVerticesEqual(const FSoftSkinVertex& V1, const FSoftSkinVertex& V2)
diff --git a/Engine/Plugins/Importers/AlembicImporter/Source/AlembicLibrary/Private/AbcImporter.cpp b/Engine/Plugins/Importers/AlembicImporter/Source/AlembicLibrary/Private/AbcImporter.cpp
--- a/Engine/Plugins/Importers/AlembicImporter/Source/AlembicLibrary/Private/AbcImporter.cpp
+++ b/Engine/Plugins/Importers/AlembicImporter/Source/AlembicLibrary/Private/AbcImporter.cpp
@@ -417,7 +417,7 @@ UGeometryCache* FAbcImporter::ImportAsGeometryCache(UObject* InParent, EObjectFl
 				Track->BeginCoding(Codec, ImportSettings->GeometryCacheSettings.bApplyConstantTopologyOptimizations && !bContainsHeterogeneousMeshes, bCalculateMotionVectors, ImportSettings->GeometryCacheSettings.bOptimizeIndexBuffers);
 				Tracks.Add(Track);
 				
-				FScopedSlowTask SlowTask((ImportSettings->SamplingSettings.FrameEnd + 1) - ImportSettings->SamplingSettings.FrameStart, FText::FromString(FString(TEXT("Importing Frames"))));
+				FScopedSlowTask SlowTask(static_cast<float>((ImportSettings->SamplingSettings.FrameEnd + 1) - ImportSettings->SamplingSettings.FrameStart), FText::FromString(FString(TEXT("Importing Frames"))));
 				SlowTask.MakeDialog(true);
 
 				const TArray<FString>& UniqueFaceSetNames = AbcFile->GetUniqueFaceSetNames();
@@ -693,8 +693,8 @@ TArray<UObject*> FAbcImporter::ImportAsSkeletalMesh(UObject* InParent, EObjectFl
 		const FFrameNumber FrameNumber = FrameRate.AsFrameNumber(AbcFile->GetImportLength());
 		Controller.SetNumberOfFrames(FrameNumber);
 
-		Sequence->ImportFileFramerate = FrameRate.AsDecimal();
-		Sequence->ImportResampleFramerate = FrameRate.AsInterval();
+		Sequence->ImportFileFramerate = static_cast<float>(FrameRate.AsDecimal());
+		Sequence->ImportResampleFramerate = static_cast<int32>(FrameRate.AsInterval());
 
 		{
 #if WITH_EDITOR
@@ -999,7 +999,7 @@ const bool FAbcImporter::CompressAnimationDataUsingPCA(const FAbcCompressionSett
 
 			// Average out vertex data
 			FBox AverageBoundingBox(ForceInit);
-			const float Multiplier = 1.0f / FMath::Max(NumSamples, 1);
+			const float Multiplier = 1.0f / static_cast<float>(FMath::Max(NumSamples, 1.f));
 			for (FVector3f& Vertex : AverageVertexData)
 			{
 				Vertex *= Multiplier;
@@ -1256,7 +1256,7 @@ const bool FAbcImporter::CompressAnimationDataUsingPCA(const FAbcCompressionSett
 
 			// Average out vertex data
 			FBox AverageBoundingBox(ForceInit);
-			const float Multiplier = 1.0f / FMath::Max(NumSamples, 1);
+			const float Multiplier = 1.0f / static_cast<float>(FMath::Max(NumSamples, 1));
 			for (TArray<FVector3f>& VertexData : AverageVertexData)
 			{
 				for (FVector3f& Vertex : VertexData)
diff --git a/Engine/Plugins/Runtime/GeometryCache/Source/GeometryCache/Classes/GeometryCacheCodecBase.h b/Engine/Plugins/Runtime/GeometryCache/Source/GeometryCache/Classes/GeometryCacheCodecBase.h
--- a/Engine/Plugins/Runtime/GeometryCache/Source/GeometryCache/Classes/GeometryCacheCodecBase.h
+++ b/Engine/Plugins/Runtime/GeometryCache/Source/GeometryCache/Classes/GeometryCacheCodecBase.h
@@ -53,8 +53,8 @@ struct FGeometryCacheCodecDecodeArguments
 		OutMeshData(SetOutMeshData)
 	{
 		// We do common validation checking here so it doesn't have to be done in every codec individually
-		checkf(FrameIdentifier >= Chunks[0].FirstFrame, TEXT("Requested frame not in the range of the provided chunks, or chunks are not properly ordered"));
-		checkf(FrameIdentifier <= Chunks.Last().LastFrame, TEXT("Requested frame not in the range of the provided chunks, or chunks are not properly ordered"));
+		checkf(static_cast<float>(FrameIdentifier) >= Chunks[0].FirstFrame, TEXT("Requested frame not in the range of the provided chunks, or chunks are not properly ordered"));
+		checkf(static_cast<float>(FrameIdentifier) <= Chunks.Last().LastFrame, TEXT("Requested frame not in the range of the provided chunks, or chunks are not properly ordered"));
 	}
 
 	FGeometryCacheTrackStreamableRenderResource &Track;
