Alembic: Fixed imprecision in frame times computation at import.

diff --git a/Engine/Plugins/Importers/AlembicImporter/Source/AlembicLibrary/Private/AbcFile.cpp b/Engine/Plugins/Importers/AlembicImporter/Source/AlembicLibrary/Private/AbcFile.cpp
--- a/Engine/Plugins/Importers/AlembicImporter/Source/AlembicLibrary/Private/AbcFile.cpp
+++ b/Engine/Plugins/Importers/AlembicImporter/Source/AlembicLibrary/Private/AbcFile.cpp
@@ -257,21 +257,21 @@ EAbcImportError FAbcFile::Import(UAbcImportSettings* InImportSettings)
 	ImportLength = static_cast<float>(FrameSpan) * TimeStep;
 
 	// Calculate time offset from start of import animation range
-	ImportTimeOffset = static_cast<float>(StartFrameIndex) * SecondsPerFrame;
+	ImportTimeOffset = static_cast<float>(StartFrameIndex) / static_cast<float>(FramesPerSecond);
 
 	// Read first-frames for both the transforms and poly meshes
 
 	bool bValidFirstFrames = true;
 	for (FAbcTransform* Transform : Transforms)
 	{
-		bValidFirstFrames &= Transform->ReadFirstFrame(static_cast<float>(StartFrameIndex) * SecondsPerFrame, StartFrameIndex);
+		bValidFirstFrames &= Transform->ReadFirstFrame(static_cast<float>(StartFrameIndex) / static_cast<float>(FramesPerSecond), StartFrameIndex);
 	}
 
 	for (FAbcPolyMesh* PolyMesh : PolyMeshes)
 	{
 		if (PolyMesh->bShouldImport)
 		{
-			bValidFirstFrames &= PolyMesh->ReadFirstFrame(static_cast<float>(StartFrameIndex) * SecondsPerFrame, StartFrameIndex);
+			bValidFirstFrames &= PolyMesh->ReadFirstFrame(static_cast<float>(StartFrameIndex) / static_cast<float>(FramesPerSecond), StartFrameIndex);
 		}
 	}	
 
@@ -501,7 +501,7 @@ void FAbcFile::ReadFrame(int32 FrameIndex, const EFrameReadFlags InFlags, const
 {
 	for (IAbcObject* Object : Objects)
 	{
-		Object->SetFrameAndTime(static_cast<float>(FrameIndex) * SecondsPerFrame, FrameIndex, InFlags, ReadIndex);
+		Object->SetFrameAndTime(static_cast<float>(FrameIndex) / static_cast<float>(FramesPerSecond), FrameIndex, InFlags, ReadIndex);
 	}
 }
 
diff --git a/Engine/Plugins/Importers/AlembicImporter/Source/AlembicLibrary/Private/AbcImporter.cpp b/Engine/Plugins/Importers/AlembicImporter/Source/AlembicLibrary/Private/AbcImporter.cpp
--- a/Engine/Plugins/Importers/AlembicImporter/Source/AlembicLibrary/Private/AbcImporter.cpp
+++ b/Engine/Plugins/Importers/AlembicImporter/Source/AlembicLibrary/Private/AbcImporter.cpp
@@ -436,7 +436,11 @@ UGeometryCache* FAbcImporter::ImportAsGeometryCache(UObject* InParent, EObjectFl
 						PolyMeshes, UniqueFaceSetNames,
 						MeshData, PreviousNumVertices, bConstantTopology, bStoreImportedVertexNumbers);
 					
-					Tracks[0]->AddMeshSample(MeshData, PolyMeshes[0]->GetTimeForFrameIndex(FrameIndex) - InAbcFile->GetImportTimeOffset(), bConstantTopology);
+					const float FrameRate = static_cast<float>(InAbcFile->GetFramerate());
+
+					// Convert frame times to frame numbers and back to time to avoid float imprecision
+					const float FrameTime = static_cast<float>(FMath::RoundToInt(PolyMeshes[0]->GetTimeForFrameIndex(FrameIndex) * FrameRate) - FMath::RoundToInt(InAbcFile->GetImportTimeOffset() * FrameRate)) / FrameRate;
+					Tracks[0]->AddMeshSample(MeshData, FrameTime, bConstantTopology);
 					
 					if (IsInGameThread())
 					{
@@ -511,6 +515,7 @@ UGeometryCache* FAbcImporter::ImportAsGeometryCache(UObject* InParent, EObjectFl
 				const int32 NumTracks = Tracks.Num();
 				TFunction<void(int32, FAbcFile*)> Callback = [this, NumTracks, &ImportPolyMeshes, &Tracks, &MaterialOffsets](int32 FrameIndex, const FAbcFile* InAbcFile)
 				{
+					const float FrameRate = static_cast<float>(InAbcFile->GetFramerate());
 					for (int32 TrackIndex = 0; TrackIndex < NumTracks; ++TrackIndex)
 					{
 						const FAbcPolyMesh* PolyMesh = ImportPolyMeshes[TrackIndex];
@@ -520,7 +525,8 @@ UGeometryCache* FAbcImporter::ImportAsGeometryCache(UObject* InParent, EObjectFl
 
 							// Generate the mesh data for this sample
 							const bool bVisible = PolyMesh->GetVisibility(FrameIndex);
-							const float FrameTime = PolyMesh->GetTimeForFrameIndex(FrameIndex) - InAbcFile->GetImportTimeOffset();
+							// Convert frame times to frame numbers and back to time to avoid float imprecision
+							const float FrameTime = static_cast<float>(FMath::RoundToInt(PolyMesh->GetTimeForFrameIndex(FrameIndex) * FrameRate) - FMath::RoundToInt(InAbcFile->GetImportTimeOffset() * FrameRate)) / FrameRate;
 							if (bVisible)
 							{
 								const bool bUseVelocitiesAsMotionVectors = ( ImportSettings->GeometryCacheSettings.MotionVectors == EAbcGeometryCacheMotionVectorsImport::ImportAbcVelocitiesAsMotionVectors );
@@ -925,12 +931,15 @@ const bool FAbcImporter::CompressAnimationDataUsingPCA(const FAbcCompressionSett
 				[this, PolyMeshesToCompress, &MinTime, &MaxTime, &NumSamples, &ObjectVertexOffsets, &ObjectIndexOffsets, &AverageVertexData, &AverageNormalData, NumPolyMeshesToCompress]
 				(int32 FrameIndex, FAbcFile* InFile)
 				{
+					const float FrameRate = static_cast<float>(AbcFile->GetFramerate());
 					for (int32 MeshIndex = 0; MeshIndex < NumPolyMeshesToCompress; ++MeshIndex)
 					{
 						FAbcPolyMesh* PolyMesh = PolyMeshesToCompress[MeshIndex];
 
-						MinTime = FMath::Min(MinTime, (float)PolyMesh->GetTimeForFrameIndex(FrameIndex) - AbcFile->GetImportTimeOffset());
-						MaxTime = FMath::Max(MaxTime, (float)PolyMesh->GetTimeForFrameIndex(FrameIndex) - AbcFile->GetImportTimeOffset());
+						// Convert frame times to frame numbers and back to time to avoid float imprecision
+						const float FrameTime = static_cast<float>(FMath::RoundToInt(PolyMesh->GetTimeForFrameIndex(FrameIndex) * FrameRate) - FMath::RoundToInt(AbcFile->GetImportTimeOffset() * FrameRate)) / FrameRate;
+						MinTime = FMath::Min(MinTime, FrameTime);
+						MaxTime = FMath::Max(MaxTime, FrameTime);
 
 						if (ObjectVertexOffsets.Num() != NumPolyMeshesToCompress)
 						{
@@ -1162,6 +1171,7 @@ const bool FAbcImporter::CompressAnimationDataUsingPCA(const FAbcCompressionSett
 				[this, NumPolyMeshesToCompress, &PolyMeshesToCompress, &MinTimes, &MaxTimes, &NumSamples, &AverageVertexData, &AverageNormalData]
 				(int32 FrameIndex, FAbcFile* InFile)
 			{
+				const float FrameRate = static_cast<float>(AbcFile->GetFramerate());
 				// Each individual object creates a compressed data object
 				for (int32 MeshIndex = 0; MeshIndex < NumPolyMeshesToCompress; ++MeshIndex)
 				{
@@ -1190,8 +1200,9 @@ const bool FAbcImporter::CompressAnimationDataUsingPCA(const FAbcCompressionSett
 						}
 					}
 
-					MinTimes[MeshIndex] = FMath::Min(MinTimes[MeshIndex], (float)PolyMesh->GetTimeForFrameIndex(FrameIndex) - AbcFile->GetImportTimeOffset());
-					MaxTimes[MeshIndex] = FMath::Max(MaxTimes[MeshIndex], (float)PolyMesh->GetTimeForFrameIndex(FrameIndex) - AbcFile->GetImportTimeOffset());
+					const float FrameTime = static_cast<float>(FMath::RoundToInt(PolyMesh->GetTimeForFrameIndex(FrameIndex) * FrameRate) - FMath::RoundToInt(AbcFile->GetImportTimeOffset() * FrameRate)) / FrameRate;
+					MinTimes[MeshIndex] = FMath::Min(MinTimes[MeshIndex], FrameTime);
+					MaxTimes[MeshIndex] = FMath::Max(MaxTimes[MeshIndex], FrameTime);
 				}
 
 				for (int32 MeshIndex = 0; MeshIndex < NumPolyMeshesToCompress; ++MeshIndex)
