GeometryCache: Fixed GeometryCache not having any animation when importing as "Flatten Tracks" but without importing the first mesh in a multi-mesh Alembic.

diff --git a/Engine/Plugins/Importers/AlembicImporter/Source/AlembicLibrary/Private/AbcImportUtilities.cpp b/Engine/Plugins/Importers/AlembicImporter/Source/AlembicLibrary/Private/AbcImportUtilities.cpp
--- a/Engine/Plugins/Importers/AlembicImporter/Source/AlembicLibrary/Private/AbcImportUtilities.cpp
+++ b/Engine/Plugins/Importers/AlembicImporter/Source/AlembicLibrary/Private/AbcImportUtilities.cpp
@@ -1628,7 +1628,7 @@ void AbcImporterUtilities::GeometryCacheDataForMeshSample(FGeometryCacheMeshData
 }
 
 void AbcImporterUtilities::MergePolyMeshesToMeshData(int32 FrameIndex, int32 FrameStart, float SecondsPerFrame, bool bUseVelocitiesAsMotionVectors,
-	const TArray<FAbcPolyMesh*>& PolyMeshes, const TArray<FString>& UniqueFaceSetNames,
+	const TArray<FAbcPolyMesh*>& PolyMeshes, const TArray<FString>& UniqueFaceSetNames, float& FrameTime,
 	FGeometryCacheMeshData& MeshData, int32& PreviousNumVertices, bool& bConstantTopology, bool bStoreImportedVertexNumbers)
 {
 	FAbcMeshSample MergedSample;
@@ -1643,6 +1643,7 @@ void AbcImporterUtilities::MergePolyMeshesToMeshData(int32 FrameIndex, int32 Fra
 			if (PolyMesh->GetVisibility(FrameIndex))
 			{
 				const FAbcMeshSample* Sample = PolyMesh->GetSample(FrameIndex);
+				FrameTime = PolyMesh->GetTimeForFrameIndex(FrameIndex);
 				AbcImporterUtilities::AppendMeshSample(&MergedSample, Sample);
 				if (PolyMesh->FaceSetNames.Num() == 0)
 				{
diff --git a/Engine/Plugins/Importers/AlembicImporter/Source/AlembicLibrary/Private/AbcImportUtilities.h b/Engine/Plugins/Importers/AlembicImporter/Source/AlembicLibrary/Private/AbcImportUtilities.h
--- a/Engine/Plugins/Importers/AlembicImporter/Source/AlembicLibrary/Private/AbcImportUtilities.h
+++ b/Engine/Plugins/Importers/AlembicImporter/Source/AlembicLibrary/Private/AbcImportUtilities.h
@@ -339,13 +339,14 @@ namespace AbcImporterUtilities
 	 * @param bUseVelocitiesAsMotionVectors		Converts the AbcPolyMesh velocities to motion vectors
 	 * @param PolyMeshes			The PolyMeshes to merge, which will be sampled at FrameIndex
 	 * @param UniqueFaceSetNames	The array of unique face set names of the PolyMeshes
+	 * @param FrameTime				The frame time that corresponds to FrameIndex
 	 * @param MeshData				The GeometryCacheMeshData where to output the merged PolyMeshes
 	 * @param PreviousNumVertices	The number of vertices in the merged PolyMeshes, used to determine if its topology is constant between 2 frames
 	 * @param bConstantTopology		Flag to indicate if the merged PolyMeshes has constant topology
 	 * @param bStoreImportedVertexNumbers Set to true when we want to store the original dcc vertex numbers for each vertex.
 	 */
 	void MergePolyMeshesToMeshData(int32 FrameIndex, int32 FrameStart, float SecondsPerFrame, bool bUseVelocitiesAsMotionVectors,
-		const TArray<FAbcPolyMesh*>& PolyMeshes, const TArray<FString>& UniqueFaceSetNames,
+		const TArray<FAbcPolyMesh*>& PolyMeshes, const TArray<FString>& UniqueFaceSetNames, float& FrameTime,
 		FGeometryCacheMeshData& MeshData, int32& PreviousNumVertices, bool& bConstantTopology, bool bStoreImportedVertexNumbers);
 
 	/** Retrieves a material from an AbcFile according to the given name and resaves it into the parent package */
diff --git a/Engine/Plugins/Importers/AlembicImporter/Source/AlembicLibrary/Private/AbcImporter.cpp b/Engine/Plugins/Importers/AlembicImporter/Source/AlembicLibrary/Private/AbcImporter.cpp
--- a/Engine/Plugins/Importers/AlembicImporter/Source/AlembicLibrary/Private/AbcImporter.cpp
+++ b/Engine/Plugins/Importers/AlembicImporter/Source/AlembicLibrary/Private/AbcImporter.cpp
@@ -422,24 +422,26 @@ UGeometryCache* FAbcImporter::ImportAsGeometryCache(UObject* InParent, EObjectFl
 
 				const TArray<FString>& UniqueFaceSetNames = AbcFile->GetUniqueFaceSetNames();
 				const TArray<FAbcPolyMesh*>& PolyMeshes = AbcFile->GetPolyMeshes();
+				TArray<float> FrameTimes;
+				FrameTimes.SetNum(ImportSettings->SamplingSettings.FrameEnd - ImportSettings->SamplingSettings.FrameStart + 1);
 				
 				const int32 NumTracks = Tracks.Num();
 				int32 PreviousNumVertices = 0;
-				TFunction<void(int32, FAbcFile*)> Callback = [this, &Tracks, &SlowTask, &UniqueFaceSetNames, &PolyMeshes, &PreviousNumVertices](int32 FrameIndex, const FAbcFile* InAbcFile)
+				TFunction<void(int32, FAbcFile*)> Callback = [this, &Tracks, &SlowTask, &UniqueFaceSetNames, &PolyMeshes, &PreviousNumVertices, &FrameTimes](int32 FrameIndex, const FAbcFile* InAbcFile)
 				{
 					const bool bUseVelocitiesAsMotionVectors = (ImportSettings->GeometryCacheSettings.MotionVectors == EAbcGeometryCacheMotionVectorsImport::ImportAbcVelocitiesAsMotionVectors);
 					FGeometryCacheMeshData MeshData;
 					bool bConstantTopology = true;
 					const bool bStoreImportedVertexNumbers = ImportSettings->GeometryCacheSettings.bStoreImportedVertexNumbers;
 
+					int32 FrameTimeIndex = FrameIndex - ImportSettings->SamplingSettings.FrameStart;
 					AbcImporterUtilities::MergePolyMeshesToMeshData(FrameIndex, ImportSettings->SamplingSettings.FrameStart, AbcFile->GetSecondsPerFrame(), bUseVelocitiesAsMotionVectors,
-						PolyMeshes, UniqueFaceSetNames,
-						MeshData, PreviousNumVertices, bConstantTopology, bStoreImportedVertexNumbers);
+						PolyMeshes, UniqueFaceSetNames, FrameTimes[FrameTimeIndex], MeshData, PreviousNumVertices, bConstantTopology, bStoreImportedVertexNumbers);
 					
 					const float FrameRate = static_cast<float>(InAbcFile->GetFramerate());
 
 					// Convert frame times to frame numbers and back to time to avoid float imprecision
-					const float FrameTime = static_cast<float>(FMath::RoundToInt(PolyMeshes[0]->GetTimeForFrameIndex(FrameIndex) * FrameRate) - FMath::RoundToInt(InAbcFile->GetImportTimeOffset() * FrameRate)) / FrameRate;
+					const float FrameTime = static_cast<float>(FMath::RoundToInt(FrameTimes[FrameTimeIndex] * FrameRate) - FMath::RoundToInt(InAbcFile->GetImportTimeOffset() * FrameRate)) / FrameRate;
 					Tracks[0]->AddMeshSample(MeshData, FrameTime, bConstantTopology);
 					
 					if (IsInGameThread())
diff --git a/Engine/Plugins/Importers/AlembicImporter/Source/AlembicLibrary/Private/AbcUtilities.cpp b/Engine/Plugins/Importers/AlembicImporter/Source/AlembicLibrary/Private/AbcUtilities.cpp
--- a/Engine/Plugins/Importers/AlembicImporter/Source/AlembicLibrary/Private/AbcUtilities.cpp
+++ b/Engine/Plugins/Importers/AlembicImporter/Source/AlembicLibrary/Private/AbcUtilities.cpp
@@ -18,6 +18,7 @@ void FAbcUtilities::GetFrameMeshData(FAbcFile& AbcFile, int32 FrameIndex, FGeome
 	const UAbcImportSettings* ImportSettings = AbcFile.GetImportSettings();
 	checkSlow(ImportSettings);
 
+	float FrameTime = 0.0f;
 	FGeometryCacheMeshData MeshData;
 	int32 PreviousNumVertices = 0;
 	bool bConstantTopology = false;
@@ -25,7 +26,7 @@ void FAbcUtilities::GetFrameMeshData(FAbcFile& AbcFile, int32 FrameIndex, FGeome
 	const bool bStoreImportedVertexNumbers = ImportSettings->GeometryCacheSettings.bStoreImportedVertexNumbers;
 
 	AbcImporterUtilities::MergePolyMeshesToMeshData(FrameIndex, 0, AbcFile.GetSecondsPerFrame(), bUseVelocitiesAsMotionVectors,
-		AbcFile.GetPolyMeshes(), AbcFile.GetUniqueFaceSetNames(), MeshData, PreviousNumVertices, bConstantTopology, bStoreImportedVertexNumbers);
+		AbcFile.GetPolyMeshes(), AbcFile.GetUniqueFaceSetNames(), FrameTime, MeshData, PreviousNumVertices, bConstantTopology, bStoreImportedVertexNumbers);
 
 	OutMeshData = MoveTemp(MeshData);
 
