Alembic:
* Fixed imported subassets (Skeleton, AnimSequence, materials) not displaying in the Content Browser.
* Fixed imported skeletal mesh animation not playing the animation.

diff --git a/Engine/Plugins/Importers/AlembicImporter/Source/AlembicLibrary/Private/AbcImportUtilities.cpp b/Engine/Plugins/Importers/AlembicImporter/Source/AlembicLibrary/Private/AbcImportUtilities.cpp
--- a/Engine/Plugins/Importers/AlembicImporter/Source/AlembicLibrary/Private/AbcImportUtilities.cpp
+++ b/Engine/Plugins/Importers/AlembicImporter/Source/AlembicLibrary/Private/AbcImportUtilities.cpp
@@ -10,11 +10,12 @@
 #include "AbcPolyMesh.h"
 #include "AssetRegistry/AssetRegistryModule.h"
 #include "Async/ParallelFor.h"
-#include "MaterialDomain.h"
-#include "MeshUtilities.h"
 #include "GeometryCacheMeshData.h"
-#include "Materials/Material.h"
 #include "Logging/TokenizedMessage.h"
+#include "MaterialDomain.h"
+#include "Materials/Material.h"
+#include "MeshUtilities.h"
+#include "PackageTools.h"
 #include "RenderMath.h"
 #include "UObject/Package.h"
 
@@ -1684,36 +1685,40 @@ UMaterialInterface* AbcImporterUtilities::RetrieveMaterial(FAbcFile& AbcFile, co
 	if (CachedMaterial)
 	{
 		Material = *CachedMaterial;
+
+		// Setup package name next to InParent
+		FString NewPackageName = FPackageName::GetLongPackagePath(InParent->GetOutermost()->GetPathName()) + TEXT("/") + MaterialName;
+		NewPackageName = UPackageTools::SanitizePackageName(NewPackageName);
+
 		// Material could have been deleted if we're overriding/reimporting an asset
 		if (Material->IsValidLowLevel())
 		{
 			if (Material->GetOuter() == GetTransientPackage())
 			{
-				UMaterial* ExistingTypedObject = FindObject<UMaterial>(InParent, *MaterialName);
-				if (!ExistingTypedObject)
-				{
-					// This is in for safety, as we do not expect this to happen
-					UObject* ExistingObject = FindObject<UObject>(InParent, *MaterialName);
-					if (ExistingObject)
-					{
-						return nullptr;
-					}
-
-					Material->Rename(*MaterialName, InParent);				
-					Material->SetFlags(Flags);
-					FAssetRegistryModule::AssetCreated(Material);
-				}
-				else
+				UPackage* Package = CreatePackage(*NewPackageName);
+				UObject* ExistingObject = FindObject<UObject>(Package, *MaterialName);
+				if (ExistingObject)
 				{
-					ExistingTypedObject->PreEditChange(nullptr);
-					Material = ExistingTypedObject;
+					return Cast<UMaterial>(ExistingObject);
 				}
+
+				Material->Rename(*MaterialName, Package);				
+				Material->SetFlags(Flags);
+				FAssetRegistryModule::AssetCreated(Material);
 			}
 		}
 		else
 		{
+			// #ueent_todo: Revisit this code path
 			// In this case recreate the material
-			Material = NewObject<UMaterial>(InParent, *MaterialName);
+			UPackage* Package = CreatePackage(*NewPackageName);
+			UObject* ExistingObject = FindObject<UObject>(Package, *MaterialName);
+			if (ExistingObject)
+			{
+				return Cast<UMaterial>(ExistingObject);
+			}
+
+			Material = NewObject<UMaterial>(Package, *MaterialName);
 			Material->SetFlags(Flags);
 			FAssetRegistryModule::AssetCreated(Material);
 		}
diff --git a/Engine/Plugins/Importers/AlembicImporter/Source/AlembicLibrary/Private/AbcImporter.cpp b/Engine/Plugins/Importers/AlembicImporter/Source/AlembicLibrary/Private/AbcImporter.cpp
--- a/Engine/Plugins/Importers/AlembicImporter/Source/AlembicLibrary/Private/AbcImporter.cpp
+++ b/Engine/Plugins/Importers/AlembicImporter/Source/AlembicLibrary/Private/AbcImporter.cpp
@@ -156,12 +156,12 @@ const EAbcImportError FAbcImporter::ImportTrackData(const int32 InNumThreads, UA
 template<typename T>
 T* FAbcImporter::CreateObjectInstance(UObject*& InParent, const FString& ObjectName, const EObjectFlags Flags)
 {
-	// Parent package to place new mesh
+	// Parent package to place new asset
 	UPackage* Package = nullptr;
 	FString NewPackageName;
 
 	// Setup package name and create one accordingly
-	NewPackageName = FPackageName::GetLongPackagePath(InParent->GetOutermost()->GetName() + TEXT("/") + ObjectName);
+	NewPackageName = FPackageName::GetLongPackagePath(InParent->GetOutermost()->GetPathName()) + TEXT("/") + ObjectName;
 	NewPackageName = UPackageTools::SanitizePackageName(NewPackageName);
 	Package = CreatePackage(*NewPackageName);
 
@@ -848,6 +848,12 @@ TArray<UObject*> FAbcImporter::ImportAsSkeletalMesh(UObject* InParent, EObjectFl
 
 void FAbcImporter::SetupMorphTargetCurves(USkeleton* Skeleton, FName ConstCurveName, UAnimSequence* Sequence, const TArray<float> &CurveValues, const TArray<float>& TimeValues, IAnimationDataController& Controller)
 {
+	// Need curve metadata for the AnimSequence to playback. Can be either on the Skeleton or SKelMesh,
+	// but by default for FBX import it's on the Skeleton so do the same for Alembic
+	const bool bMaterialCurve = false;
+	const bool bMorphTargetCurve = true;
+	Skeleton->AccumulateCurveMetaData(ConstCurveName, bMaterialCurve, bMorphTargetCurve);
+
 	FAnimationCurveIdentifier CurveId(ConstCurveName, ERawCurveTrackTypes::RCT_Float);
 	Controller.AddCurve(CurveId);
 
