Fix non-unity
#rnx

diff --git a/Engine/Plugins/Importers/AlembicImporter/Source/AlembicLibrary/Private/AbcImportUtilities.cpp b/Engine/Plugins/Importers/AlembicImporter/Source/AlembicLibrary/Private/AbcImportUtilities.cpp
--- a/Engine/Plugins/Importers/AlembicImporter/Source/AlembicLibrary/Private/AbcImportUtilities.cpp
+++ b/Engine/Plugins/Importers/AlembicImporter/Source/AlembicLibrary/Private/AbcImportUtilities.cpp
@@ -15,6 +15,7 @@
 #include "MaterialDomain.h"
 #include "Materials/Material.h"
 #include "MeshUtilities.h"
+#include "Misc/PackageName.h"
 #include "PackageTools.h"
 #include "RenderMath.h"
 #include "UObject/Package.h"
