Alembic: Made the Alembic importer support GeometryCache by permission (always allowed by default).

diff --git a/Engine/Plugins/Importers/AlembicImporter/Source/AlembicImporter/Private/AlembicImportFactory.cpp b/Engine/Plugins/Importers/AlembicImporter/Source/AlembicImporter/Private/AlembicImportFactory.cpp
--- a/Engine/Plugins/Importers/AlembicImporter/Source/AlembicImporter/Private/AlembicImportFactory.cpp
+++ b/Engine/Plugins/Importers/AlembicImporter/Source/AlembicImporter/Private/AlembicImportFactory.cpp
@@ -25,6 +25,8 @@
 #include "ImportUtils/StaticMeshImportUtils.h"
 #include "Subsystems/ImportSubsystem.h"
 #include "UObject/UObjectIterator.h"
+#include "AssetToolsModule.h"
+#include "Misc/NamePermissionList.h"
 
 #include UE_INLINE_GENERATED_CPP_BY_NAME(AlembicImportFactory)
 
@@ -191,12 +193,24 @@ UObject* UAlembicImportFactory::FactoryCreateFile(UClass* InClass, UObject* InPa
 		}
 		else
 		{
+			static bool bGeometryCacheAllowed = []() -> bool
+			{
+				IAssetTools& AssetTools = FAssetToolsModule::GetModule().Get();
+				TSharedPtr<FPathPermissionList> AssetClassPermissionList = AssetTools.GetAssetClassPathPermissionList(EAssetClassAction::ImportAsset);
+				if (AssetClassPermissionList && AssetClassPermissionList->HasFiltering())
+				{
+					return AssetClassPermissionList->PassesFilter(TEXT("/Script/GeometryCache.GeometryCache"));
+				}
+
+				return true;
+			}();
+
 			if (ImportSettings->ImportType == EAlembicImportType::StaticMesh)
 			{
 				const TArray<UObject*> ResultStaticMeshes = ImportStaticMesh(Importer, InParent, Flags);
 				ResultAssets.Append(ResultStaticMeshes);
 			}
-			else if (ImportSettings->ImportType == EAlembicImportType::GeometryCache)
+			else if (bGeometryCacheAllowed && ImportSettings->ImportType == EAlembicImportType::GeometryCache)
 			{
 				UObject* GeometryCache = ImportGeometryCache(Importer, InParent, Flags);
 				if (GeometryCache)
@@ -218,8 +232,12 @@ UObject* UAlembicImportFactory::FactoryCreateFile(UClass* InClass, UObject* InPa
 			{
 				FAssetRegistryModule::AssetCreated(Object);
 				GEditor->GetEditorSubsystem<UImportSubsystem>()->BroadcastAssetPostImport(this, Object);
-				Object->MarkPackageDirty();
-				Object->PostEditChange();
+				// Avoid redoing a PostEditChange since the importer will have most likely done it already
+				if (!Object->GetOutermost()->IsDirty())
+				{
+					Object->MarkPackageDirty();
+					Object->PostEditChange();
+				}
 				AdditionalImportedObjects.Add(Object);
 			}
 		}
diff --git a/Engine/Plugins/Importers/AlembicImporter/Source/AlembicLibrary/Private/AbcImportSettingsCustomization.cpp b/Engine/Plugins/Importers/AlembicImporter/Source/AlembicLibrary/Private/AbcImportSettingsCustomization.cpp
--- a/Engine/Plugins/Importers/AlembicImporter/Source/AlembicLibrary/Private/AbcImportSettingsCustomization.cpp
+++ b/Engine/Plugins/Importers/AlembicImporter/Source/AlembicLibrary/Private/AbcImportSettingsCustomization.cpp
@@ -9,6 +9,8 @@
 #include "IDetailChildrenBuilder.h"
 #include "IDetailPropertyRow.h"
 #include "PropertyRestriction.h"
+#include "AssetToolsModule.h"
+#include "Misc/NamePermissionList.h"
 
 void FAbcImportSettingsCustomization::CustomizeDetails(IDetailLayoutBuilder& LayoutBuilder)
 {
@@ -36,20 +38,47 @@ void FAbcImportSettingsCustomization::CustomizeDetails(IDetailLayoutBuilder& Lay
 
 	UAbcImportSettings* CurrentSettings = SettingsObject ? Cast<UAbcImportSettings>(SettingsObject->Get()) : nullptr;
 
-	if (CurrentSettings && CurrentSettings->bReimport)
+	if (CurrentSettings)
 	{
-		UEnum* ImportTypeEnum = StaticEnum<EAlembicImportType>();		
-		static FText RestrictReason = NSLOCTEXT("AlembicImportFactory", "ReimportRestriction", "Unable to change type while reimporting");
-		TSharedPtr<FPropertyRestriction> EnumRestriction = MakeShareable(new FPropertyRestriction(RestrictReason));
+		if (CurrentSettings->bReimport)
+		{
+			UEnum* ImportTypeEnum = StaticEnum<EAlembicImportType>();
+			static FText RestrictReason = NSLOCTEXT("AlembicImportFactory", "ReimportRestriction", "Unable to change type while reimporting");
+			TSharedPtr<FPropertyRestriction> EnumRestriction = MakeShareable(new FPropertyRestriction(RestrictReason));
 
-		for (uint8 EnumIndex = 0; EnumIndex < (ImportTypeEnum->GetMaxEnumValue() + 1); ++EnumIndex)
+			for (uint8 EnumIndex = 0; EnumIndex < (ImportTypeEnum->GetMaxEnumValue() + 1); ++EnumIndex)
+			{
+				if (EnumValue != EnumIndex)
+				{
+					EnumRestriction->AddDisabledValue(ImportTypeEnum->GetNameStringByIndex(EnumIndex));
+				}
+			}
+			ImportType->AddRestriction(EnumRestriction.ToSharedRef());
+		}
+		else
 		{
-			if (EnumValue != EnumIndex)
+			static bool bGeometryCacheAllowed = []() -> bool
 			{
-				EnumRestriction->AddDisabledValue(ImportTypeEnum->GetNameStringByIndex(EnumIndex));
+				IAssetTools& AssetTools = FAssetToolsModule::GetModule().Get();
+				TSharedPtr<FPathPermissionList> AssetClassPermissionList = AssetTools.GetAssetClassPathPermissionList(EAssetClassAction::ImportAsset);
+				if (AssetClassPermissionList && AssetClassPermissionList->HasFiltering())
+				{
+					return AssetClassPermissionList->PassesFilter(TEXT("/Script/GeometryCache.GeometryCache"));
+				}
+
+				return true;
+			}();
+
+			UEnum* ImportTypeEnum = StaticEnum<EAlembicImportType>();
+			static FText RestrictReason = NSLOCTEXT("AlembicImportFactory", "ImportRestriction", "Geometry Cache import disabled");
+			TSharedPtr<FPropertyRestriction> EnumRestriction = MakeShareable(new FPropertyRestriction(RestrictReason));
+
+			if (!bGeometryCacheAllowed)
+			{
+				EnumRestriction->AddHiddenValue(ImportTypeEnum->GetNameStringByValue((uint8)EAlembicImportType::GeometryCache));
+				ImportType->AddRestriction(EnumRestriction.ToSharedRef());
 			}
-		}		
-		ImportType->AddRestriction(EnumRestriction.ToSharedRef());
+		}
 	}	
 }
 
