Alembic: Improved support for importing large mesh animation as skeletal mesh, for example cloth. The preferred import type is still geometry cache since it's faster and use less memory to import.
* Added import progress feedback for skeletal mesh and non-flattened geometry cache import.

diff --git a/Engine/Plugins/Importers/AlembicImporter/Source/AlembicLibrary/Private/AbcImportUtilities.cpp b/Engine/Plugins/Importers/AlembicImporter/Source/AlembicLibrary/Private/AbcImportUtilities.cpp
--- a/Engine/Plugins/Importers/AlembicImporter/Source/AlembicLibrary/Private/AbcImportUtilities.cpp
+++ b/Engine/Plugins/Importers/AlembicImporter/Source/AlembicLibrary/Private/AbcImportUtilities.cpp
@@ -16,6 +16,7 @@
 #include "Materials/Material.h"
 #include "MeshUtilities.h"
 #include "Misc/PackageName.h"
+#include "Misc/ScopedSlowTask.h"
 #include "PackageTools.h"
 #include "RenderMath.h"
 #include "UObject/Package.h"
@@ -1190,12 +1191,12 @@ void AbcImporterUtilities::PropogateMatrixTransformationToSample(FAbcMeshSample*
 }
 
 void AbcImporterUtilities::GenerateDeltaFrameDataMatrix(const TArray<FVector3f>& FrameVertexData, const TArray<FVector3f>& FrameNormalData, const TArray<FVector3f>& AverageVertexData, const TArray<FVector3f>& AverageNormalData,
-	const int32 SampleIndex, const int32 AverageVertexOffset, const int32 AverageIndexOffset, const FVector& SamplePositionOffset, TArray<float>& OutGeneratedMatrix, TArray<float>& OutGeneratedNormalsMatrix)
+	const int32 SampleIndex, const int32 AverageVertexOffset, const int32 AverageIndexOffset, const FVector& SamplePositionOffset, TArray64<float>& OutGeneratedMatrix, TArray64<float>& OutGeneratedNormalsMatrix)
 {
-	const uint32 NumVertices = FrameVertexData.Num();
-	const uint32 VertexOffset = SampleIndex * AverageVertexData.Num() * 3;
+	const int32 NumVertices = FrameVertexData.Num();
+	const int64 VertexOffset = int64(SampleIndex) * AverageVertexData.Num() * 3;
 
-	for (uint32 VertexIndex = 0; VertexIndex < NumVertices; ++VertexIndex)
+	for (int32 VertexIndex = 0; VertexIndex < NumVertices; ++VertexIndex)
 	{
 		const int32 ComponentIndexOffset = (VertexIndex + AverageVertexOffset) * 3;
 		const FVector AverageDifference = ((FVector)AverageVertexData[VertexIndex + AverageVertexOffset] + SamplePositionOffset) - (FVector)FrameVertexData[VertexIndex];
@@ -1205,10 +1206,10 @@ void AbcImporterUtilities::GenerateDeltaFrameDataMatrix(const TArray<FVector3f>&
 		OutGeneratedMatrix[VertexOffset + ComponentIndexOffset + 2] = static_cast<float>(AverageDifference.Z);
 	}
 
-	const uint32 NumIndices = FrameNormalData.Num();
-	const uint32 IndexOffset = SampleIndex * AverageNormalData.Num() * 3;
+	const int32 NumIndices = FrameNormalData.Num();
+	const int64 IndexOffset = int64(SampleIndex) * AverageNormalData.Num() * 3;
 
-	for (uint32 Index = 0; Index < NumIndices; ++Index)
+	for (int32 Index = 0; Index < NumIndices; ++Index)
 	{
 		const int32 ComponentIndexOffset = (Index + AverageIndexOffset) * 3;
 		const FVector AverageNormal = (FVector)AverageNormalData[Index + AverageIndexOffset] - (FVector)FrameNormalData[Index];
@@ -1220,8 +1221,11 @@ void AbcImporterUtilities::GenerateDeltaFrameDataMatrix(const TArray<FVector3f>&
 }
 
 void AbcImporterUtilities::GenerateCompressedMeshData(FCompressedAbcData& CompressedData, const uint32 NumUsedSingularValues, const uint32 NumSamples,
-	const TArrayView<float>& BasesMatrix, const TArrayView<float>& NormalsBasesMatrix, const TArray<float>& BasesWeights, const float SampleTimeStep, const float StartTime)
+	const TArrayView64<float>& BasesMatrix, const TArrayView64<float>& NormalsBasesMatrix, const TArray64<float>& BasesWeights, const float SampleTimeStep, const float StartTime)
 {
+	FScopedSlowTask SlowTask(static_cast<float>(NumUsedSingularValues), FText::FromString(FString(TEXT("Generating bases"))));
+	SlowTask.MakeDialog();
+
 	// Allocate base sample data	
 	CompressedData.BaseSamples.AddZeroed(NumUsedSingularValues);
 	CompressedData.CurveValues.AddZeroed(NumUsedSingularValues);
@@ -1232,12 +1236,12 @@ void AbcImporterUtilities::GenerateCompressedMeshData(FCompressedAbcData& Compre
 	{
 		FAbcMeshSample* Base = new FAbcMeshSample(*CompressedData.AverageSample);
 
-		const uint32 NumVertices = Base->Vertices.Num();
-		const uint32 NumMatrixRows = NumVertices * 3;
-		const int32 BaseOffset = BaseIndex * NumMatrixRows;
-		for (uint32 Index = 0; Index < NumVertices; ++Index)
+		const int32 NumVertices = Base->Vertices.Num();
+		const int32 NumMatrixRows = NumVertices * 3;
+		const int64 BaseOffset = int64(BaseIndex) * NumMatrixRows;
+		for (int32 Index = 0; Index < NumVertices; ++Index)
 		{
-			const int32 IndexOffset = BaseOffset + (Index * 3);
+			const int64 IndexOffset = BaseOffset + (Index * 3);
 			FVector3f& BaseVertex = Base->Vertices[Index];
 
 			BaseVertex.X -= BasesMatrix[IndexOffset + 0];
@@ -1245,12 +1249,12 @@ void AbcImporterUtilities::GenerateCompressedMeshData(FCompressedAbcData& Compre
 			BaseVertex.Z -= BasesMatrix[IndexOffset + 2];
 		}
 
-		const uint32 NumIndices = Base->Indices.Num();
-		const uint32 NumNormalsMatrixRows = NumIndices * 3;
-		const int32 BaseIndexOffset = BaseIndex * NumNormalsMatrixRows;
-		for (uint32 Index = 0; Index < NumIndices; ++Index)
+		const int32 NumIndices = Base->Indices.Num();
+		const int32 NumNormalsMatrixRows = NumIndices * 3;
+		const int64 BaseIndexOffset = int64(BaseIndex) * NumNormalsMatrixRows;
+		for (int32 Index = 0; Index < NumIndices; ++Index)
 		{
-			const int32 IndexOffset = BaseIndexOffset + (Index * 3);
+			const int64 IndexOffset = BaseIndexOffset + (Index * 3);
 			FVector3f& BaseNormal = Base->Normals[Index];
 
 			BaseNormal.X -= NormalsBasesMatrix[IndexOffset + 0];
@@ -1267,13 +1271,15 @@ void AbcImporterUtilities::GenerateCompressedMeshData(FCompressedAbcData& Compre
 		TimeValues.Reserve(NumSamples);
 
 		// Use original number of singular values to index into the array (otherwise we would be reading incorrect data if NumUsedSingularValues != the original number
-		const uint32 OriginalNumberOfSingularValues = BasesWeights.Num() / NumSamples;
+		const int64 OriginalNumberOfSingularValues = BasesWeights.Num() / NumSamples;
 		// Should be possible to rearrange the data so this can become a memcpy
 		for (uint32 CurveSampleIndex = 0; CurveSampleIndex < NumSamples; ++CurveSampleIndex)
 		{
 			CurveValues.Add(BasesWeights[BaseIndex + (OriginalNumberOfSingularValues * CurveSampleIndex)]);
 			TimeValues.Add(StartTime + (SampleTimeStep * static_cast<float>(CurveSampleIndex)));
 		}
+
+		SlowTask.EnterProgressFrame(1.f);
 	}
 }
 
diff --git a/Engine/Plugins/Importers/AlembicImporter/Source/AlembicLibrary/Private/AbcImportUtilities.h b/Engine/Plugins/Importers/AlembicImporter/Source/AlembicLibrary/Private/AbcImportUtilities.h
--- a/Engine/Plugins/Importers/AlembicImporter/Source/AlembicLibrary/Private/AbcImportUtilities.h
+++ b/Engine/Plugins/Importers/AlembicImporter/Source/AlembicLibrary/Private/AbcImportUtilities.h
@@ -297,11 +297,11 @@ namespace AbcImporterUtilities
 
 	/** Generates the delta frame data for the given average and frame vertex data */
 	void GenerateDeltaFrameDataMatrix(const TArray<FVector3f>& FrameVertexData, const TArray<FVector3f>& FrameNormalData, const TArray<FVector3f>& AverageVertexData, const TArray<FVector3f>& AverageNormalData,
-		const int32 SampleIndex, const int32 AverageVertexOffset, const int32 AverageIndexOffset, const FVector& SamplePositionOffset, TArray<float>& OutGeneratedMatrix, TArray<float>& OutGeneratedNormalsMatrix);
+		const int32 SampleIndex, const int32 AverageVertexOffset, const int32 AverageIndexOffset, const FVector& SamplePositionOffset, TArray64<float>& OutGeneratedMatrix, TArray64<float>& OutGeneratedNormalsMatrix);
 
 	/** Populates compressed data structure from the result PCA compression bases and weights */
 	void GenerateCompressedMeshData(FCompressedAbcData& CompressedData, const uint32 NumUsedSingularValues, const uint32 NumSamples,
-		const TArrayView<float>& BasesMatrix, const TArrayView<float>& NormalsBasesMatrix, const TArray<float>& BasesWeights, const float SampleTimeStep, const float StartTime);
+		const TArrayView64<float>& BasesMatrix, const TArrayView64<float>& NormalsBasesMatrix, const TArray64<float>& BasesWeights, const float SampleTimeStep, const float StartTime);
 
 	void CalculateNewStartAndEndFrameIndices(const float FrameStepRatio, int32& InOutStartFrameIndex, int32& InOutEndFrameIndex );
 
diff --git a/Engine/Plugins/Importers/AlembicImporter/Source/AlembicLibrary/Private/AbcImporter.cpp b/Engine/Plugins/Importers/AlembicImporter/Source/AlembicLibrary/Private/AbcImporter.cpp
--- a/Engine/Plugins/Importers/AlembicImporter/Source/AlembicLibrary/Private/AbcImporter.cpp
+++ b/Engine/Plugins/Importers/AlembicImporter/Source/AlembicLibrary/Private/AbcImporter.cpp
@@ -1,72 +1,57 @@
 // Copyright Epic Games, Inc. All Rights Reserved.
 
 #include "AbcImporter.h"
+
+#include "AbcAssetImportData.h"
+#include "AbcFile.h"
 #include "AbcImportLogger.h"
+#include "AbcImportUtilities.h"
 #include "AbcPolyMesh.h"
+#include "Animation/AnimSequence.h"
 #include "Animation/Skeleton.h"
-
+#include "Async/ParallelFor.h"
 #include "BoneWeights.h"
-
-#if PLATFORM_WINDOWS
-#include "AbcPolyMesh.h"
-#include "Windows/WindowsHWrapper.h"
-#endif
-
-THIRD_PARTY_INCLUDES_START
-#include "Animation/AnimData/AnimDataModel.h"
-#include <Alembic/AbcCoreFactory/All.h>
-#include "Animation/AnimData/IAnimationDataController.h"
-#include <Alembic/AbcCoreOgawa/All.h>
-THIRD_PARTY_INCLUDES_END
-
-#include "Engine/StaticMeshSourceData.h"
-#include "Misc/Paths.h"
+#include "ComponentReregisterContext.h"
+#include "Editor.h"
+#include "EigenHelper.h"
+#include "Engine/SkeletalMesh.h"
+#include "Engine/SkinnedAssetCommon.h"
+#include "Engine/StaticMesh.h"
 #include "GeometryCache.h"
-#include "Misc/ScopedSlowTask.h"
-
+#include "GeometryCacheCodecV1.h"
 #include "GeometryCacheComponent.h"
-#include "PackageTools.h"
 #include "GeometryCacheMeshData.h"
-#include "StaticMeshAttributes.h"
 #include "GeometryCacheTrackStreamable.h"
-#include "StaticMeshOperations.h"
 #include "Logging/TokenizedMessage.h"
-#include "ObjectTools.h"
-
-#include "Engine/StaticMesh.h"
-#include "Engine/SkeletalMesh.h"
-#include "Engine/SkinnedAssetCommon.h"
-#include "Animation/AnimSequence.h"
-#include "Misc/PackageName.h"
-#include "Rendering/SkeletalMeshModel.h"
-
-#include "AbcImportUtilities.h"
-
-#include "MeshUtilities.h"
 #include "MaterialDomain.h"
 #include "Materials/Material.h"
+#include "MeshBudgetProjectSettings.h"
+#include "MeshUtilities.h"
+#include "Misc/MessageDialog.h"
+#include "Misc/PackageName.h"
+#include "Misc/Paths.h"
+#include "Misc/ScopedSlowTask.h"
 #include "Modules/ModuleManager.h"
-
-#include "Async/ParallelFor.h"
-
-#include "EigenHelper.h"
-
-#include "AbcAssetImportData.h"
-#include "AbcFile.h"
-
-#include "ComponentReregisterContext.h"
-#include "GeometryCacheCodecV1.h"
+#include "ObjectTools.h"
+#include "PackageTools.h"
 #include "RenderMath.h"
+#include "Rendering/SkeletalMeshModel.h"
+#include "StaticMeshAttributes.h"
+#include "StaticMeshOperations.h"
 #include "Subsystems/AssetEditorSubsystem.h"
-#include "Editor.h"
-
-#include "Rendering/SkeletalMeshLODModel.h"
 #include "UObject/MetaData.h"
 
-#if WITH_EDITOR
-#include "MeshBudgetProjectSettings.h"
+#if PLATFORM_WINDOWS
+#include "Windows/WindowsHWrapper.h"
 #endif
 
+THIRD_PARTY_INCLUDES_START
+#include "Animation/AnimData/AnimDataModel.h"
+#include <Alembic/AbcCoreFactory/All.h>
+#include "Animation/AnimData/IAnimationDataController.h"
+#include <Alembic/AbcCoreOgawa/All.h>
+THIRD_PARTY_INCLUDES_END
+
 #define LOCTEXT_NAMESPACE "AbcImporter"
 
 DEFINE_LOG_CATEGORY_STATIC(LogAbcImporter, Verbose, All);
@@ -270,9 +255,8 @@ UStaticMesh* FAbcImporter::CreateStaticMeshFromSample(UObject* InParent, const F
 		//Set the Imported version before calling the build
 		StaticMesh->ImportVersion = EImportStaticMeshVersion::LastVersion;
 
-#if WITH_EDITOR
 		FMeshBudgetProjectSettingsUtils::SetLodGroupForStaticMesh(StaticMesh);
-#endif
+
 		// Build the static mesh (using the build setting etc.) this generates correct tangents using the extracting smoothing group along with the imported Normals data
 		StaticMesh->Build(false);
 
@@ -418,7 +402,7 @@ UGeometryCache* FAbcImporter::ImportAsGeometryCache(UObject* InParent, EObjectFl
 				Tracks.Add(Track);
 				
 				FScopedSlowTask SlowTask(static_cast<float>((ImportSettings->SamplingSettings.FrameEnd + 1) - ImportSettings->SamplingSettings.FrameStart), FText::FromString(FString(TEXT("Importing Frames"))));
-				SlowTask.MakeDialog(true);
+				SlowTask.MakeDialog();
 
 				const TArray<FString>& UniqueFaceSetNames = AbcFile->GetUniqueFaceSetNames();
 				const TArray<FAbcPolyMesh*>& PolyMeshes = AbcFile->GetPolyMeshes();
@@ -427,7 +411,8 @@ UGeometryCache* FAbcImporter::ImportAsGeometryCache(UObject* InParent, EObjectFl
 				
 				const int32 NumTracks = Tracks.Num();
 				int32 PreviousNumVertices = 0;
-				TFunction<void(int32, FAbcFile*)> Callback = [this, &Tracks, &SlowTask, &UniqueFaceSetNames, &PolyMeshes, &PreviousNumVertices, &FrameTimes](int32 FrameIndex, const FAbcFile* InAbcFile)
+				float CompletedFrames = 0.0f;
+				TFunction<void(int32, FAbcFile*)> Callback = [this, &Tracks, &SlowTask, &UniqueFaceSetNames, &PolyMeshes, &PreviousNumVertices, &FrameTimes, &CompletedFrames](int32 FrameIndex, const FAbcFile* InAbcFile)
 				{
 					const bool bUseVelocitiesAsMotionVectors = (ImportSettings->GeometryCacheSettings.MotionVectors == EAbcGeometryCacheMotionVectorsImport::ImportAbcVelocitiesAsMotionVectors);
 					FGeometryCacheMeshData MeshData;
@@ -444,9 +429,11 @@ UGeometryCache* FAbcImporter::ImportAsGeometryCache(UObject* InParent, EObjectFl
 					const float FrameTime = static_cast<float>(FMath::RoundToInt(FrameTimes[FrameTimeIndex] * FrameRate) - FMath::RoundToInt(InAbcFile->GetImportTimeOffset() * FrameRate)) / FrameRate;
 					Tracks[0]->AddMeshSample(MeshData, FrameTime, bConstantTopology);
 					
+					++CompletedFrames;
 					if (IsInGameThread())
 					{
-						SlowTask.EnterProgressFrame(1.0f);
+						SlowTask.EnterProgressFrame(CompletedFrames);
+						CompletedFrames = 0.0f;
 					}					
 				};
 
@@ -514,8 +501,12 @@ UGeometryCache* FAbcImporter::ImportAsGeometryCache(UObject* InParent, EObjectFl
 					}
 				}
 
+				FScopedSlowTask SlowTask(static_cast<float>((ImportSettings->SamplingSettings.FrameEnd + 1) - ImportSettings->SamplingSettings.FrameStart), FText::FromString(FString(TEXT("Importing Frames"))));
+				SlowTask.MakeDialog();
+
 				const int32 NumTracks = Tracks.Num();
-				TFunction<void(int32, FAbcFile*)> Callback = [this, NumTracks, &ImportPolyMeshes, &Tracks, &MaterialOffsets](int32 FrameIndex, const FAbcFile* InAbcFile)
+				float CompletedFrames = 0.0f;
+				TFunction<void(int32, FAbcFile*)> Callback = [this, NumTracks, &ImportPolyMeshes, &Tracks, &MaterialOffsets, &SlowTask, &CompletedFrames](int32 FrameIndex, const FAbcFile* InAbcFile)
 				{
 					const float FrameRate = static_cast<float>(InAbcFile->GetFramerate());
 					for (int32 TrackIndex = 0; TrackIndex < NumTracks; ++TrackIndex)
@@ -542,6 +533,13 @@ UGeometryCache* FAbcImporter::ImportAsGeometryCache(UObject* InParent, EObjectFl
 							Track->AddVisibilitySample(bVisible, FrameTime);
 						}
 					}
+					
+					++CompletedFrames;
+					if (IsInGameThread())
+					{
+						SlowTask.EnterProgressFrame(CompletedFrames);
+						CompletedFrames = 0.0f;
+					}					
 				};
 
 				AbcFile->ProcessFrames(Callback, EFrameReadFlags::ApplyMatrix);
@@ -705,11 +703,9 @@ TArray<UObject*> FAbcImporter::ImportAsSkeletalMesh(UObject* InParent, EObjectFl
 		Sequence->ImportResampleFramerate = static_cast<int32>(FrameRate.AsInterval());
 
 		{
-#if WITH_EDITOR
 			// When ScopedPostEditChange goes out of scope, it will call SkeletalMesh->PostEditChange()
 			// while preventing any call to that within the scope
 			FScopedSkeletalMeshPostEditChange ScopedPostEditChange(SkeletalMesh);
-#endif
 
 			for (FCompressedAbcData& CompressedData : CompressedMeshData)
 			{
@@ -933,10 +929,14 @@ const bool FAbcImporter::CompressAnimationDataUsingPCA(const FAbcCompressionSett
 			float MaxTime = -FLT_MAX;
 			int32 NumSamples = 0;
 
+			TUniquePtr<FScopedSlowTask> SlowTask = MakeUnique<FScopedSlowTask>(static_cast<float>((ImportSettings->SamplingSettings.FrameEnd + 1) - ImportSettings->SamplingSettings.FrameStart + 1), FText::FromString(FString(TEXT("Merging meshes"))));
+			SlowTask->MakeDialog();
+
 			TArray<uint32> ObjectVertexOffsets;
 			TArray<uint32> ObjectIndexOffsets;
+			float CompletedFrames = 0.0f;
 			TFunction<void(int32, FAbcFile*)> MergedMeshesFunc =
-				[this, PolyMeshesToCompress, &MinTime, &MaxTime, &NumSamples, &ObjectVertexOffsets, &ObjectIndexOffsets, &AverageVertexData, &AverageNormalData, NumPolyMeshesToCompress]
+				[this, PolyMeshesToCompress, &MinTime, &MaxTime, &NumSamples, &ObjectVertexOffsets, &ObjectIndexOffsets, &AverageVertexData, &AverageNormalData, NumPolyMeshesToCompress, &SlowTask, &CompletedFrames]
 				(int32 FrameIndex, FAbcFile* InFile)
 				{
 					const float FrameRate = static_cast<float>(AbcFile->GetFramerate());
@@ -971,6 +971,13 @@ const bool FAbcImporter::CompressAnimationDataUsingPCA(const FAbcCompressionSett
 					}
 
 					++NumSamples;
+					++CompletedFrames;
+
+					if (IsInGameThread())
+					{
+						SlowTask->EnterProgressFrame(CompletedFrames);
+						CompletedFrames = 0.0f;
+					}					
 				};
 
 			EFrameReadFlags Flags = EFrameReadFlags::PositionAndNormalOnly;
@@ -987,19 +994,41 @@ const bool FAbcImporter::CompressAnimationDataUsingPCA(const FAbcCompressionSett
 
 				const int32 NumFrames = AbcFile->GetEndFrameIndex() - AbcFile->GetStartFrameIndex() + 1;
 				const uint64 NumMatrixElements = uint64(AverageVertexData.Num()) * 3 * NumFrames;
+				bool bTriggerWarning = false;
+				int32 NumElementsWarning = 0;
+				uint64 NumMatrixElementsWarning = 0;
 				if (!IntFitsIn<int32>(NumMatrixElements))
 				{
-					UE_LOG(LogAbcImporter, Error, TEXT("Vertex matrix has too many elements (%llu) because the mesh has too many vertices (%d) and/or the animation has too many frames (%d). Try importing as GeometryCache instead."),
-						NumMatrixElements, AverageVertexData.Num(), NumFrames);
-					return false;
+					NumElementsWarning = AverageVertexData.Num();
+					NumMatrixElementsWarning = NumMatrixElements;
+					bTriggerWarning = true;
 				}
 
 				const uint64 NumNormalsMatrixElements = uint64(AverageNormalData.Num()) * 3 * NumFrames;
 				if (!IntFitsIn<int32>(NumNormalsMatrixElements))
 				{
-					UE_LOG(LogAbcImporter, Error, TEXT("Normal matrix has too many elements (%llu) because the mesh has too many vertices (%d) and/or the animation has too many frames (%d). Try importing as GeometryCache instead."),
-						NumNormalsMatrixElements, AverageNormalData.Num(), NumFrames);
-					return false;
+					if (AverageNormalData.Num() > NumElementsWarning || NumNormalsMatrixElements > NumMatrixElementsWarning)
+					{
+						NumElementsWarning = AverageNormalData.Num();
+						NumMatrixElementsWarning = NumNormalsMatrixElements;
+					}
+					bTriggerWarning = true;
+				}
+
+				if (bTriggerWarning)
+				{
+					UE_LOG(LogAbcImporter, Warning, TEXT("Vertex matrix has %llu elements because the mesh has %d vertices and the animation has %d frames. This can cause the import to take a long time and use a lot of memory."),
+						NumMatrixElementsWarning, NumElementsWarning, NumFrames);
+
+					const FText Title = LOCTEXT("AbcSkelMeshImportWarningTitle", "Proceed with import?");
+					const FText Message = LOCTEXT("AbcSkelMeshImportWarningMessage", "Warning: Due to the mesh size and animation length, the import may take a long time and may run out of memory and crash. Do you want to continue?\n"
+						"If not, you may try reducing the animation import range (will use less memory) or use No Compression as the Base Calculation Type (faster, but will use more memory) or import as Geometry Cache.");
+
+					const EAppReturnType::Type DialogResponse = FMessageDialog::Open(EAppMsgType::OkCancel, EAppReturnType::Ok, Message, Title);
+					if (DialogResponse != EAppReturnType::Ok)
+					{
+						return false;
+					}
 				}
 
 				AverageVertexData.Reset();
@@ -1013,6 +1042,7 @@ const bool FAbcImporter::CompressAnimationDataUsingPCA(const FAbcCompressionSett
 			}
 
 			AbcFile->ProcessFrames(MergedMeshesFunc, Flags);
+			SlowTask.Reset();
 
 			// Average out vertex data
 			FBox AverageBoundingBox(ForceInit);
@@ -1055,10 +1085,10 @@ const bool FAbcImporter::CompressAnimationDataUsingPCA(const FAbcCompressionSett
 			const uint32 NumIndices = AverageNormalData.Num();
 			const uint32 NumNormalsMatrixRows = NumIndices * 3;
 
-			TArray<float> OriginalMatrix;
+			TArray64<float> OriginalMatrix;
 			OriginalMatrix.AddZeroed(NumMatrixRows * NumSamples);
 
-			TArray<float> OriginalNormalsMatrix;
+			TArray64<float> OriginalNormalsMatrix;
 			OriginalNormalsMatrix.AddZeroed(NumNormalsMatrixRows * NumSamples);
 
 			if (bEnableSamplesOffsets)
@@ -1067,10 +1097,15 @@ const bool FAbcImporter::CompressAnimationDataUsingPCA(const FAbcCompressionSett
 				SamplesOffsets.GetValue().AddZeroed(NumSamples);
 			}
 
+			SlowTask = MakeUnique<FScopedSlowTask>(static_cast<float>((ImportSettings->SamplingSettings.FrameEnd + 1) - ImportSettings->SamplingSettings.FrameStart), FText::FromString(FString(TEXT("Generating matrices"))));
+			SlowTask->MakeDialog();
+
+			CompletedFrames = 0.0f;
+
 			uint32 GenerateMatrixSampleIndex = 0;
 			TFunction<void(int32, FAbcFile*)> GenerateMatrixFunc =
 				[this, PolyMeshesToCompress, NumPolyMeshesToCompress, &OriginalMatrix, &OriginalNormalsMatrix, &AverageVertexData, &AverageNormalData, &NumSamples,
-					&ObjectVertexOffsets, &ObjectIndexOffsets, &GenerateMatrixSampleIndex, &AverageSampleCenter]
+					&ObjectVertexOffsets, &ObjectIndexOffsets, &GenerateMatrixSampleIndex, &AverageSampleCenter, &SlowTask, &CompletedFrames]
 				(int32 FrameIndex, FAbcFile* InFile)
 				{
 					FVector SampleOffset = FVector::ZeroVector;
@@ -1105,14 +1140,22 @@ const bool FAbcImporter::CompressAnimationDataUsingPCA(const FAbcCompressionSett
 					}
 
 					++GenerateMatrixSampleIndex;
+					++CompletedFrames;
+
+					if (IsInGameThread())
+					{
+						SlowTask->EnterProgressFrame(CompletedFrames);
+						CompletedFrames = 0.0f;
+					}					
 				};
 
 			AbcFile->ProcessFrames(GenerateMatrixFunc, Flags);
+			SlowTask.Reset();
 
 			// Perform compression
-			TArray<float> OutU, OutV, OutNormalsU;
-			TArrayView<float> BasesMatrix;
-			TArrayView<float> NormalsBasesMatrix;
+			TArray64<float> OutU, OutV, OutNormalsU;
+			TArrayView64<float> BasesMatrix;
+			TArrayView64<float> NormalsBasesMatrix;
 			uint32 NumUsedSingularValues = NumSamples;
 
 			if (InCompressionSettings.BaseCalculationType != EBaseCalculationType::NoCompression)
@@ -1173,10 +1216,13 @@ const bool FAbcImporter::CompressAnimationDataUsingPCA(const FAbcCompressionSett
 			AverageVertexData.AddDefaulted(NumPolyMeshesToCompress);
 			AverageNormalData.AddDefaulted(NumPolyMeshesToCompress);
 			
-			
+			TUniquePtr<FScopedSlowTask> SlowTask = MakeUnique<FScopedSlowTask>(static_cast<float>((ImportSettings->SamplingSettings.FrameEnd + 1) - ImportSettings->SamplingSettings.FrameStart + 1), FText::FromString(FString(TEXT("Processing meshes"))));
+			SlowTask->MakeDialog();
+
 			int32 NumSamples = 0;
+			float CompletedFrames = 0.0f;
 			TFunction<void(int32, FAbcFile*)> IndividualMeshesFunc =
-				[this, NumPolyMeshesToCompress, &PolyMeshesToCompress, &MinTimes, &MaxTimes, &NumSamples, &AverageVertexData, &AverageNormalData]
+				[this, NumPolyMeshesToCompress, &PolyMeshesToCompress, &MinTimes, &MaxTimes, &NumSamples, &AverageVertexData, &AverageNormalData, &SlowTask, &CompletedFrames]
 				(int32 FrameIndex, FAbcFile* InFile)
 			{
 				const float FrameRate = static_cast<float>(AbcFile->GetFramerate());
@@ -1224,6 +1270,13 @@ const bool FAbcImporter::CompressAnimationDataUsingPCA(const FAbcCompressionSett
 				}
 
 				++NumSamples;
+				++CompletedFrames;
+
+				if (IsInGameThread())
+				{
+					SlowTask->EnterProgressFrame(CompletedFrames);
+					CompletedFrames = 0.0f;
+				}					
 			};
 
 			EFrameReadFlags Flags = EFrameReadFlags::PositionAndNormalOnly;
@@ -1239,21 +1292,42 @@ const bool FAbcImporter::CompressAnimationDataUsingPCA(const FAbcCompressionSett
 				AbcFile->CleanupFrameData(0);
 
 				const int32 NumFrames = AbcFile->GetEndFrameIndex() - AbcFile->GetStartFrameIndex() + 1;
+				bool bTriggerWarning = false;
+				int32 NumElementsWarning = 0;
+				uint64 NumMatrixElementsWarning = 0;
 				for (int32 MeshIndex = 0; MeshIndex < NumPolyMeshesToCompress; ++MeshIndex)
 				{
 					const uint64 NumMatrixElements = uint64(AverageVertexData[MeshIndex].Num()) * 3 * NumFrames;
 					if (!IntFitsIn<int32>(NumMatrixElements))
 					{
-						UE_LOG(LogAbcImporter, Error, TEXT("Vertex matrix has too many elements (%llu) because the mesh has too many vertices (%d) and/or the animation has too many frames (%d). Try importing as GeometryCache instead."),
-						NumMatrixElements, AverageVertexData[MeshIndex].Num(), NumFrames);
-						return false;
+						NumElementsWarning = AverageVertexData[MeshIndex].Num();
+						NumMatrixElementsWarning = NumMatrixElements;
+						bTriggerWarning = true;
+						break;
 					}
 
 					const uint64 NumNormalsMatrixElements = uint64(AverageNormalData[MeshIndex].Num()) * 3 * NumFrames;
 					if (!IntFitsIn<int32>(NumNormalsMatrixElements))
 					{
-						UE_LOG(LogAbcImporter, Error, TEXT("Normal matrix has too many elements (%llu) because the mesh has too many vertices (%d) and/or the animation has too many frames (%d). Try importing as GeometryCache instead."),
-						NumNormalsMatrixElements, AverageNormalData[MeshIndex].Num(), NumFrames);
+						NumElementsWarning = AverageNormalData[MeshIndex].Num();
+						NumMatrixElementsWarning = NumNormalsMatrixElements;
+						bTriggerWarning = true;
+						break;
+					}
+				}
+
+				if (bTriggerWarning)
+				{
+					UE_LOG(LogAbcImporter, Warning, TEXT("Vertex matrix has %llu elements because the mesh has %d vertices and the animation has %d frames. This can cause the import to take a long time and use a lot of memory."),
+						NumMatrixElementsWarning, NumElementsWarning, NumFrames);
+
+					const FText Title = LOCTEXT("AbcSkelMeshImportWarningTitle", "Proceed with import?");
+					const FText Message = LOCTEXT("AbcSkelMeshImportWarningMessage", "Warning: Due to the mesh size and animation length, the import may take a long time and may run out of memory and crash. Do you want to continue?\n"
+						"If not, you may try reducing the animation import range (will use less memory) or use No Compression as the Base Calculation Type (faster, but will use more memory) or import as Geometry Cache.");
+
+					const EAppReturnType::Type DialogResponse = FMessageDialog::Open(EAppMsgType::OkCancel, EAppReturnType::Ok, Message, Title);
+					if (DialogResponse != EAppReturnType::Ok)
+					{
 						return false;
 					}
 				}
@@ -1272,6 +1346,7 @@ const bool FAbcImporter::CompressAnimationDataUsingPCA(const FAbcCompressionSett
 			}
 
 			AbcFile->ProcessFrames(IndividualMeshesFunc, Flags);
+			SlowTask.Reset();
 
 			// Average out vertex data
 			FBox AverageBoundingBox(ForceInit);
@@ -1286,15 +1361,15 @@ const bool FAbcImporter::CompressAnimationDataUsingPCA(const FAbcCompressionSett
 			}
 			const FVector AverageSampleCenter = AverageBoundingBox.GetCenter();
 
-			TArray<TArray<float>> Matrices;
-			TArray<TArray<float>> NormalsMatrices;
+			TArray<TArray64<float>> Matrices;
+			TArray<TArray64<float>> NormalsMatrices;
 			for (int32 MeshIndex = 0; MeshIndex < NumPolyMeshesToCompress; ++MeshIndex)
 			{
 				Matrices.AddDefaulted();
-				Matrices[MeshIndex].AddZeroed(AverageVertexData[MeshIndex].Num() * 3 * NumSamples);
+				Matrices[MeshIndex].AddZeroed(int64(AverageVertexData[MeshIndex].Num()) * 3 * NumSamples);
 
 				NormalsMatrices.AddDefaulted();
-				NormalsMatrices[MeshIndex].AddZeroed(AverageNormalData[MeshIndex].Num() * 3 * NumSamples);
+				NormalsMatrices[MeshIndex].AddZeroed(int64(AverageNormalData[MeshIndex].Num()) * 3 * NumSamples);
 			}
 
 			if (bEnableSamplesOffsets)
@@ -1309,9 +1384,12 @@ const bool FAbcImporter::CompressAnimationDataUsingPCA(const FAbcCompressionSett
 				SamplesOffsets.GetValue().AddDefaulted(NumSamples);
 			}
 
+			SlowTask = MakeUnique<FScopedSlowTask>(static_cast<float>((ImportSettings->SamplingSettings.FrameEnd + 1) - ImportSettings->SamplingSettings.FrameStart), FText::FromString(FString(TEXT("Generating matrices"))));
+			SlowTask->MakeDialog();
+
 			uint32 GenerateMatrixSampleIndex = 0;
 			TFunction<void(int32, FAbcFile*)> GenerateMatrixFunc =
-				[this, NumPolyMeshesToCompress, &Matrices, &NormalsMatrices, &GenerateMatrixSampleIndex, &PolyMeshesToCompress, &AverageVertexData, &AverageNormalData, &AverageSampleCenter]
+				[this, NumPolyMeshesToCompress, &Matrices, &NormalsMatrices, &GenerateMatrixSampleIndex, &PolyMeshesToCompress, &AverageVertexData, &AverageNormalData, &AverageSampleCenter, &SlowTask, &CompletedFrames]
 				(int32 FrameIndex, FAbcFile* InFile)
 				{
 					// Compute on bounding box for the sample, which will include all the meshes
@@ -1353,16 +1431,24 @@ const bool FAbcImporter::CompressAnimationDataUsingPCA(const FAbcCompressionSett
 					}
 
 					++GenerateMatrixSampleIndex;
-				};
+					++CompletedFrames;
+
+					if (IsInGameThread())
+					{
+						SlowTask->EnterProgressFrame(CompletedFrames);
+						CompletedFrames = 0.0f;
+					}					
+			};
 
 			AbcFile->ProcessFrames(GenerateMatrixFunc, Flags);
+			SlowTask.Reset();
 
 			for (int32 MeshIndex = 0; MeshIndex < NumPolyMeshesToCompress; ++MeshIndex)
 			{
 				// Perform compression
-				TArray<float> OutU, OutV, OutNormalsU;
-				TArrayView<float> BasesMatrix;
-				TArrayView<float> NormalsBasesMatrix;
+				TArray64<float> OutU, OutV, OutNormalsU;
+				TArrayView64<float> BasesMatrix;
+				TArrayView64<float> NormalsBasesMatrix;
 
 				const int32 NumVertices = AverageVertexData[MeshIndex].Num();
 				const int32 NumIndices = AverageNormalData[MeshIndex].Num();
@@ -1468,25 +1554,25 @@ const bool FAbcImporter::CompressAnimationDataUsingPCA(const FAbcCompressionSett
 	return bResult;
 }
 
-void FAbcImporter::CompareCompressionResult(const TArray<float>& OriginalMatrix, const uint32 NumSamples, const uint32 NumUsedSingularValues, const TArrayView<float>& OutU, const TArray<float>& OutV, const float Tolerance)
+void FAbcImporter::CompareCompressionResult(const TArray64<float>& OriginalMatrix, const uint32 NumSamples, const uint32 NumUsedSingularValues, const TArrayView64<float>& OutU, const TArray64<float>& OutV, const float Tolerance)
 {
 	if (NumSamples == 0)
 	{
 		return;
 	}
 
-	const uint32 NumRows = OriginalMatrix.Num() / NumSamples;
+	const uint32 NumRows = IntCastChecked<uint32>(OriginalMatrix.Num() / NumSamples);
 
-	TArray<float> ComparisonMatrix;
+	TArray64<float> ComparisonMatrix;
 	ComparisonMatrix.AddZeroed(OriginalMatrix.Num());
 	for (uint32 SampleIndex = 0; SampleIndex < NumSamples; ++SampleIndex)
 	{
-		const int32 SampleOffset = (SampleIndex * NumRows);
-		const int32 CurveOffset = (SampleIndex * NumUsedSingularValues);
+		const int64 SampleOffset = (int64(SampleIndex) * NumRows);
+		const int64 CurveOffset = (int64(SampleIndex) * NumUsedSingularValues);
 		for (uint32 BaseIndex = 0; BaseIndex < NumUsedSingularValues; ++BaseIndex)
 		{
-			const int32 BaseOffset = (BaseIndex * NumRows);
-			for (uint32 RowIndex = 0; RowIndex < NumRows; RowIndex++)
+			const int64 BaseOffset = (int64(BaseIndex) * NumRows);
+			for (uint32 RowIndex = 0; RowIndex < NumRows; ++RowIndex)
 			{
 				ComparisonMatrix[RowIndex + SampleOffset] += OutU[RowIndex + BaseOffset] * OutV[BaseIndex + CurveOffset];
 			}
@@ -1494,30 +1580,34 @@ void FAbcImporter::CompareCompressionResult(const TArray<float>& OriginalMatrix,
 	}
 
 	// Compare arrays
-	for (int32 i = 0; i < ComparisonMatrix.Num(); ++i)
+	for (int64 i = 0; i < ComparisonMatrix.Num(); ++i)
 	{
 		ensureMsgf(FMath::IsNearlyEqual(OriginalMatrix[i], ComparisonMatrix[i], Tolerance), TEXT("Difference of %2.10f found"), FMath::Abs(OriginalMatrix[i] - ComparisonMatrix[i]));
 	}
 }
 
-const int32 FAbcImporter::PerformSVDCompression(const TArray<float>& OriginalMatrix, const TArray<float>& OriginalNormalsMatrix, const uint32 NumSamples, const float InPercentage, const int32 InFixedNumValue,
-	TArray<float>& OutU, TArray<float>& OutNormalsU, TArray<float>& OutV)
+const int32 FAbcImporter::PerformSVDCompression(const TArray64<float>& OriginalMatrix, const TArray64<float>& OriginalNormalsMatrix, const uint32 NumSamples, const float InPercentage, const int32 InFixedNumValue,
+	TArray64<float>& OutU, TArray64<float>& OutNormalsU, TArray64<float>& OutV)
 {
-	const int32 NumRows = OriginalMatrix.Num() / NumSamples;
+	FScopedSlowTask SlowTask(4.0f, FText::FromString(FString(TEXT("Decomposing animation"))));
+	SlowTask.MakeDialog();
+
+	const int32 NumRows = IntCastChecked<int32>(OriginalMatrix.Num() / NumSamples);
 
-	TArray<float> OutS;
+	TArray64<float> OutS;
 	EigenHelpers::PerformSVD(OriginalMatrix, NumRows, NumSamples, OutU, OutV, OutS);
+	SlowTask.EnterProgressFrame(1.0f);
 
 	// Now we have the new basis data we have to construct the correct morph target data and curves
 	const float PercentageBasesUsed = InPercentage;
-	const int32 NumNonZeroSingularValues = OutS.Num();
-	const int32 NumUsedSingularValues = (InFixedNumValue != 0) ? FMath::Min(InFixedNumValue, (int32)OutS.Num()) : (int32)((float)NumNonZeroSingularValues * PercentageBasesUsed);
+	const int32 NumNonZeroSingularValues = IntCastChecked<int32>(OutS.Num());
+	const int32 NumUsedSingularValues = (InFixedNumValue != 0) ? FMath::Min(InFixedNumValue, NumNonZeroSingularValues) : (int32)((float)NumNonZeroSingularValues * PercentageBasesUsed);
 
 	// Pre-multiply the bases with it's singular values
-	ParallelFor(NumUsedSingularValues, [&](int32 ValueIndex)
+	ParallelFor(NumUsedSingularValues, [&](int64 ValueIndex)
 	{
 		const float Multiplier = OutS[ValueIndex];
-		const int32 ValueOffset = ValueIndex * NumRows;
+		const int64 ValueOffset = ValueIndex * NumRows;
 
 		for (int32 RowIndex = 0; RowIndex < NumRows; ++RowIndex)
 		{
@@ -1531,20 +1621,23 @@ const int32 FAbcImporter::PerformSVDCompression(const TArray<float>& OriginalMat
 	//
 	// This takes into account that OutNormalsU should be already scaled by what would be OutNormalsS, just like OutU is scaled by OutS
 
-	const int32 NormalsNumRows = OriginalNormalsMatrix.Num() / NumSamples;
+	const int32 NormalsNumRows = IntCastChecked<int32>(OriginalNormalsMatrix.Num() / NumSamples);
 
 	Eigen::MatrixXf NormalsMatrix;
 	EigenHelpers::ConvertArrayToEigenMatrix(OriginalNormalsMatrix, NormalsNumRows, NumSamples, NormalsMatrix);
+	SlowTask.EnterProgressFrame(1.0f);
 
-	const uint32 OutVNumRows = OutV.Num() / NumSamples;
+	const int32 OutVNumRows = IntCastChecked<uint32>(OutV.Num() / NumSamples);
 
 	Eigen::MatrixXf VMatrix;
 	EigenHelpers::ConvertArrayToEigenMatrix(OutV, OutVNumRows, NumSamples, VMatrix);
+	SlowTask.EnterProgressFrame(1.0f);
 
 	Eigen::MatrixXf NormalsUMatrix = NormalsMatrix * VMatrix.transpose();
 
 	uint32 OutNumColumns, OutNumRows;
 	EigenHelpers::ConvertEigenMatrixToArray(NormalsUMatrix, OutNormalsU, OutNumColumns, OutNumRows);
+	SlowTask.EnterProgressFrame(1.0f);
 
 	UE_LOG(LogAbcImporter, Log, TEXT("Decomposed animation and reconstructed with %i number of bases (full %i, percentage %f, calculated %i)"), NumUsedSingularValues, OutS.Num(), PercentageBasesUsed * 100.0f, NumUsedSingularValues);	
 	
diff --git a/Engine/Plugins/Importers/AlembicImporter/Source/AlembicLibrary/Private/EigenHelper.h b/Engine/Plugins/Importers/AlembicImporter/Source/AlembicLibrary/Private/EigenHelper.h
--- a/Engine/Plugins/Importers/AlembicImporter/Source/AlembicLibrary/Private/EigenHelper.h
+++ b/Engine/Plugins/Importers/AlembicImporter/Source/AlembicLibrary/Private/EigenHelper.h
@@ -38,7 +38,7 @@ namespace EigenHelpers
 	}
 
 	/** Converts a float array into an Eigen Matrix */
-	static void ConvertArrayToEigenMatrix(const TArray<float>& InArray, const int32 InRows, const int32 InColumns, Eigen::MatrixXf& OutMatrix)
+	static void ConvertArrayToEigenMatrix(const TArray64<float>& InArray, const int32 InRows, const int32 InColumns, Eigen::MatrixXf& OutMatrix)
 	{
 		OutMatrix.resize(InRows, InColumns);
 
@@ -48,7 +48,7 @@ namespace EigenHelpers
 		// Copy matrix data
 		for (int32 ColumnIndex = 0; ColumnIndex < InColumns; ++ColumnIndex)
 		{
-			const int32 ColumnOffset = ColumnIndex * InRows;
+			const int64 ColumnOffset = int64(ColumnIndex) * InRows;
 			for (int32 RowIndex = 0; RowIndex < InRows; ++RowIndex)
 			{
 				OutMatrix(RowIndex, ColumnIndex) = InArray[ColumnOffset + RowIndex];
@@ -58,11 +58,11 @@ namespace EigenHelpers
 	}
 
 	/** Converts an Eigen Matrix into a float array */
-	static void ConvertEigenMatrixToArray(const Eigen::MatrixXf& InMatrix, TArray<float>& OutArray, uint32& OutColumns, uint32& OutRows )
+	static void ConvertEigenMatrixToArray(const Eigen::MatrixXf& InMatrix, TArray64<float>& OutArray, uint32& OutColumns, uint32& OutRows )
 	{
 		OutColumns = (int32)InMatrix.cols();
 		OutRows = (int32)InMatrix.rows();
-		const uint32 TotalSize = OutRows * OutColumns;
+		const int64 TotalSize = InMatrix.rows() * InMatrix.cols();
 		OutArray.Empty(TotalSize);
 		OutArray.AddZeroed(TotalSize);
 
@@ -71,7 +71,7 @@ namespace EigenHelpers
 #else
 		for (int32 ColumnIndex = 0; ColumnIndex < OutColumns; ++ColumnIndex)
 		{
-			const int32 ColumnOffset = ColumnIndex * OutRows;
+			const int64 ColumnOffset = int64(ColumnIndex) * OutRows;
 			for (int32 RowIndex = 0; RowIndex < OutRows; ++RowIndex)
 			{
 				OutArray[ColumnOffset + RowIndex] = InMatrix(RowIndex, ColumnIndex);
@@ -90,7 +90,7 @@ namespace EigenHelpers
 	}
 
 	/** Performs Singular value decomposition on the given matrix and returns respective the calculated U, V and S Matrix */
-	static void PerformSVD(const TArray<float>& InMatrix, const int32 InRows, const int32 InColumns, TArray<float>& OutU, TArray<float>& OutV, TArray<float>& OutS)
+	static void PerformSVD(const TArray64<float>& InMatrix, const int32 InRows, const int32 InColumns, TArray64<float>& OutU, TArray64<float>& OutV, TArray64<float>& OutS)
 	{
 		Eigen::MatrixXf EigenMatrix;
 		ConvertArrayToEigenMatrix(InMatrix, InRows, InColumns, EigenMatrix);
diff --git a/Engine/Plugins/Importers/AlembicImporter/Source/AlembicLibrary/Public/AbcImporter.h b/Engine/Plugins/Importers/AlembicImporter/Source/AlembicLibrary/Public/AbcImporter.h
--- a/Engine/Plugins/Importers/AlembicImporter/Source/AlembicLibrary/Public/AbcImporter.h
+++ b/Engine/Plugins/Importers/AlembicImporter/Source/AlembicLibrary/Public/AbcImporter.h
@@ -191,11 +191,11 @@ private:
 	/** Compresses the imported animation data, returns true if compression was successful and compressed data was populated */
 	const bool CompressAnimationDataUsingPCA(const FAbcCompressionSettings& InCompressionSettings, const bool bRunComparison = false);	
 	/** Performs the actual SVD compression to retrieve the bases and weights used to set up the Skeletal mesh's morph targets */
-	const int32 PerformSVDCompression(const TArray<float>& OriginalMatrix, const TArray<float>& OriginalNormalsMatrix, const uint32 NumSamples, const float InPercentage, const int32 InFixedNumValue,
-		TArray<float>& OutU, TArray<float>& OutNormalsU, TArray<float>& OutV);
+	const int32 PerformSVDCompression(const TArray64<float>& OriginalMatrix, const TArray64<float>& OriginalNormalsMatrix, const uint32 NumSamples, const float InPercentage, const int32 InFixedNumValue,
+		TArray64<float>& OutU, TArray64<float>& OutNormalsU, TArray64<float>& OutV);
 
 	/** Functionality for comparing the matrices and calculating the difference from the original animation */
-	void CompareCompressionResult(const TArray<float>& OriginalMatrix, const uint32 NumSamples, const uint32 NumUsedSingularValues, const TArrayView<float>& OutU, const TArray<float>& OutV, const float Tolerance);
+	void CompareCompressionResult(const TArray64<float>& OriginalMatrix, const uint32 NumSamples, const uint32 NumUsedSingularValues, const TArrayView64<float>& OutU, const TArray64<float>& OutV, const float Tolerance);
 	
 	/** Build a skeletal mesh from the PCA compressed data */
 	bool BuildSkeletalMesh(FSkeletalMeshLODModel& LODModel, const FReferenceSkeleton& RefSkeleton, FAbcMeshSample* Sample, TArray<int32>& OutMorphTargetVertexRemapping, TArray<int32>& OutUsedVertexIndicesForMorphs);
