Alembic: Fixed failing tests due to bad normals when computing smoothing groups. Those tests have been re-enabled.

diff --git a/Engine/Plugins/Importers/AlembicImporter/Source/AlembicLibrary/Private/AbcImportUtilities.cpp b/Engine/Plugins/Importers/AlembicImporter/Source/AlembicLibrary/Private/AbcImportUtilities.cpp
--- a/Engine/Plugins/Importers/AlembicImporter/Source/AlembicLibrary/Private/AbcImportUtilities.cpp
+++ b/Engine/Plugins/Importers/AlembicImporter/Source/AlembicLibrary/Private/AbcImportUtilities.cpp
@@ -798,7 +798,8 @@ void AbcImporterUtilities::GenerateSmoothingGroupsIndices(FAbcMeshSample* MeshSa
 		Offset += NumVertsForFace;
 
 		// Store the averaged face normal
-		FaceNormals.Add(FaceNormal.GetSafeNormal());
+		FaceNormal.Normalize();
+		FaceNormals.Add(FaceNormal);
 	}
 		
 	MeshSample->NumSmoothingGroups = 0;
diff --git a/Engine/Plugins/Importers/AlembicImporter/Source/AlembicLibrary/Private/AbcImporter.cpp b/Engine/Plugins/Importers/AlembicImporter/Source/AlembicLibrary/Private/AbcImporter.cpp
--- a/Engine/Plugins/Importers/AlembicImporter/Source/AlembicLibrary/Private/AbcImporter.cpp
+++ b/Engine/Plugins/Importers/AlembicImporter/Source/AlembicLibrary/Private/AbcImporter.cpp
@@ -1388,6 +1388,7 @@ const bool FAbcImporter::CompressAnimationDataUsingPCA(const FAbcCompressionSett
 			SlowTask->MakeDialog();
 
 			uint32 GenerateMatrixSampleIndex = 0;
+			CompletedFrames = 0.0f;
 			TFunction<void(int32, FAbcFile*)> GenerateMatrixFunc =
 				[this, NumPolyMeshesToCompress, &Matrices, &NormalsMatrices, &GenerateMatrixSampleIndex, &PolyMeshesToCompress, &AverageVertexData, &AverageNormalData, &AverageSampleCenter, &SlowTask, &CompletedFrames]
 				(int32 FrameIndex, FAbcFile* InFile)
