Alembic. Fixed failing SK_sampling_PerStep test. The Per X Frames and Per Time Step sampling types could generate non-integral frame rate but it was forced to an integer causing computation errors. The FrameRate is now a float.

diff --git a/Engine/Plugins/Importers/AlembicImporter/Source/AlembicLibrary/Private/AbcFile.cpp b/Engine/Plugins/Importers/AlembicImporter/Source/AlembicLibrary/Private/AbcFile.cpp
--- a/Engine/Plugins/Importers/AlembicImporter/Source/AlembicLibrary/Private/AbcFile.cpp
+++ b/Engine/Plugins/Importers/AlembicImporter/Source/AlembicLibrary/Private/AbcFile.cpp
@@ -253,25 +253,25 @@ EAbcImportError FAbcFile::Import(UAbcImportSettings* InImportSettings)
 	}
 
 	SecondsPerFrame = TimeStep;
-	FramesPerSecond = TimeStep > 0.f ? FMath::RoundToInt(1.f / TimeStep) : 30;
+	FramesPerSecond = TimeStep > 0.f ? (1.f / TimeStep) : 30.f;
 	ImportLength = static_cast<float>(FrameSpan) * TimeStep;
 
 	// Calculate time offset from start of import animation range
-	ImportTimeOffset = static_cast<float>(StartFrameIndex) / static_cast<float>(FramesPerSecond);
+	ImportTimeOffset = static_cast<float>(StartFrameIndex) / FramesPerSecond;
 
 	// Read first-frames for both the transforms and poly meshes
 
 	bool bValidFirstFrames = true;
 	for (FAbcTransform* Transform : Transforms)
 	{
-		bValidFirstFrames &= Transform->ReadFirstFrame(static_cast<float>(StartFrameIndex) / static_cast<float>(FramesPerSecond), StartFrameIndex);
+		bValidFirstFrames &= Transform->ReadFirstFrame(static_cast<float>(StartFrameIndex) / FramesPerSecond, StartFrameIndex);
 	}
 
 	for (FAbcPolyMesh* PolyMesh : PolyMeshes)
 	{
 		if (PolyMesh->bShouldImport)
 		{
-			bValidFirstFrames &= PolyMesh->ReadFirstFrame(static_cast<float>(StartFrameIndex) / static_cast<float>(FramesPerSecond), StartFrameIndex);
+			bValidFirstFrames &= PolyMesh->ReadFirstFrame(static_cast<float>(StartFrameIndex) / FramesPerSecond, StartFrameIndex);
 		}
 	}	
 
@@ -501,7 +501,7 @@ void FAbcFile::ReadFrame(int32 FrameIndex, const EFrameReadFlags InFlags, const
 {
 	for (IAbcObject* Object : Objects)
 	{
-		Object->SetFrameAndTime(static_cast<float>(FrameIndex) / static_cast<float>(FramesPerSecond), FrameIndex, InFlags, ReadIndex);
+		Object->SetFrameAndTime(static_cast<float>(FrameIndex) / FramesPerSecond, FrameIndex, InFlags, ReadIndex);
 	}
 }
 
@@ -640,7 +640,7 @@ const int32 FAbcFile::GetImportNumFrames() const
 	return EndFrameIndex - StartFrameIndex;
 }
 
-const int32 FAbcFile::GetFramerate() const
+const float FAbcFile::GetFramerate() const
 {
 	return FramesPerSecond;
 }
diff --git a/Engine/Plugins/Importers/AlembicImporter/Source/AlembicLibrary/Private/AbcImporter.cpp b/Engine/Plugins/Importers/AlembicImporter/Source/AlembicLibrary/Private/AbcImporter.cpp
--- a/Engine/Plugins/Importers/AlembicImporter/Source/AlembicLibrary/Private/AbcImporter.cpp
+++ b/Engine/Plugins/Importers/AlembicImporter/Source/AlembicLibrary/Private/AbcImporter.cpp
@@ -423,7 +423,7 @@ UGeometryCache* FAbcImporter::ImportAsGeometryCache(UObject* InParent, EObjectFl
 					AbcImporterUtilities::MergePolyMeshesToMeshData(FrameIndex, ImportSettings->SamplingSettings.FrameStart, AbcFile->GetSecondsPerFrame(), bUseVelocitiesAsMotionVectors,
 						PolyMeshes, UniqueFaceSetNames, FrameTimes[FrameTimeIndex], MeshData, PreviousNumVertices, bConstantTopology, bStoreImportedVertexNumbers);
 					
-					const float FrameRate = static_cast<float>(InAbcFile->GetFramerate());
+					const float FrameRate = InAbcFile->GetFramerate();
 
 					// Convert frame times to frame numbers and back to time to avoid float imprecision
 					const float FrameTime = static_cast<float>(FMath::RoundToInt(FrameTimes[FrameTimeIndex] * FrameRate) - FMath::RoundToInt(InAbcFile->GetImportTimeOffset() * FrameRate)) / FrameRate;
@@ -694,7 +694,7 @@ TArray<UObject*> FAbcImporter::ImportAsSkeletalMesh(UObject* InParent, EObjectFl
 		Controller.OpenBracket(LOCTEXT("ImportAsSkeletalMesh", "Importing Alembic Animation"));
 		Controller.InitializeModel();
 
-		const FFrameRate FrameRate(AbcFile->GetFramerate(), 1);
+		const FFrameRate FrameRate(FMath::RoundToInt(AbcFile->GetFramerate()), 1);
 		Controller.SetFrameRate(FrameRate);	
 		const FFrameNumber FrameNumber = FrameRate.AsFrameNumber(AbcFile->GetImportLength());
 		Controller.SetNumberOfFrames(FrameNumber);
diff --git a/Engine/Plugins/Importers/AlembicImporter/Source/AlembicLibrary/Public/AbcFile.h b/Engine/Plugins/Importers/AlembicImporter/Source/AlembicLibrary/Public/AbcFile.h
--- a/Engine/Plugins/Importers/AlembicImporter/Source/AlembicLibrary/Public/AbcFile.h
+++ b/Engine/Plugins/Importers/AlembicImporter/Source/AlembicLibrary/Public/AbcFile.h
@@ -65,7 +65,7 @@ public:
 	const float GetImportTimeOffset() const;
 	const float GetImportLength() const;
 	const int32 GetImportNumFrames() const;
-	const int32 GetFramerate() const;
+	const float GetFramerate() const;
 	const float GetSecondsPerFrame() const;
 	const FBoxSphereBounds& GetArchiveBounds() const;
 	const bool ContainsHeterogeneousMeshes() const;
@@ -139,7 +139,7 @@ protected:
 	/** Total (max) number of frames in the Alembic file */
 	int32 NumFrames;
 	/** Frames per second (retrieved and specified in top Alembic object) */
-	int32 FramesPerSecond;
+	float FramesPerSecond;
 	/** Seconds per frame (calculated according to FPS) */
 	float SecondsPerFrame;
 
