PR #10719: Fix Alembic import conversion presets. Max/Maya preset names have been switched since their transform were reversed.

diff --git a/Engine/Plugins/Importers/AlembicImporter/Source/AlembicLibrary/Private/AbcImportSettingsCustomization.cpp b/Engine/Plugins/Importers/AlembicImporter/Source/AlembicLibrary/Private/AbcImportSettingsCustomization.cpp
--- a/Engine/Plugins/Importers/AlembicImporter/Source/AlembicLibrary/Private/AbcImportSettingsCustomization.cpp
+++ b/Engine/Plugins/Importers/AlembicImporter/Source/AlembicLibrary/Private/AbcImportSettingsCustomization.cpp
@@ -216,7 +216,7 @@ void FAbcConversionSettingsCustomization::OnConversionPresetChanged()
 		// Set values to specified preset
 		switch (Settings->Preset)
 		{
-			case EAbcConversionPreset::Maya:
+			case EAbcConversionPreset::Max:
 			{
 				Settings->bFlipU = false;
 				Settings->bFlipV = true;
@@ -225,7 +225,7 @@ void FAbcConversionSettingsCustomization::OnConversionPresetChanged()
 				break;
 			}
 
-			case EAbcConversionPreset::Max:
+			case EAbcConversionPreset::Maya:
 			{
 				Settings->bFlipU = false;
 				Settings->bFlipV = true;
diff --git a/Engine/Plugins/Importers/AlembicImporter/Source/AlembicLibrary/Public/AbcImportSettings.h b/Engine/Plugins/Importers/AlembicImporter/Source/AlembicLibrary/Public/AbcImportSettings.h
--- a/Engine/Plugins/Importers/AlembicImporter/Source/AlembicLibrary/Public/AbcImportSettings.h
+++ b/Engine/Plugins/Importers/AlembicImporter/Source/AlembicLibrary/Public/AbcImportSettings.h
@@ -218,7 +218,7 @@ struct FAbcConversionSettings
 		, bFlipU(false)
 		, bFlipV(true)
 		, Scale(FVector(1.0f, -1.0f, 1.0f))
-		, Rotation(FVector::ZeroVector)
+		, Rotation(FVector(90.0f, 0.0f, 0.0f))
 	{}
 
 	/** Currently preset that should be applied */
