PR #10721: AlembicImport: frame range discovery for animated visibility. This allows importing animation of visibility only.

diff --git a/Engine/Plugins/Importers/AlembicImporter/Source/AlembicLibrary/Private/AbcFile.cpp b/Engine/Plugins/Importers/AlembicImporter/Source/AlembicLibrary/Private/AbcFile.cpp
--- a/Engine/Plugins/Importers/AlembicImporter/Source/AlembicLibrary/Private/AbcFile.cpp
+++ b/Engine/Plugins/Importers/AlembicImporter/Source/AlembicLibrary/Private/AbcFile.cpp
@@ -452,6 +452,34 @@ void FAbcFile::TraverseAbcHierarchy(const Alembic::Abc::IObject& InObject, IAbcO
 		}
 	}
 
+	// Check for animated visibility property
+	{
+		const Alembic::Abc::ICompoundProperty CompoundProperty = InObject.getProperties();
+
+		if (CompoundProperty.getPropertyHeader(Alembic::AbcGeom::kVisibilityPropertyName))
+		{
+			const Alembic::AbcGeom::IVisibilityProperty VisibilityProperty = Alembic::AbcGeom::IVisibilityProperty(CompoundProperty, Alembic::AbcGeom::kVisibilityPropertyName);
+
+			int32 NumSamplesVisibility = VisibilityProperty.getNumSamples();
+
+			if (NumSamplesVisibility > 1)
+			{
+				const Alembic::AbcCoreAbstract::TimeSamplingPtr TimeSampler = VisibilityProperty.getTimeSampling();
+				const Alembic::AbcCoreAbstract::TimeSamplingType SamplingType = TimeSampler->getTimeSamplingType();
+
+				const float StartTimeVisibility = (float)TimeSampler->getSampleTime(0);
+				const int32 MinFrameIndexVisibility = FMath::CeilToInt(StartTimeVisibility / (float)SamplingType.getTimePerCycle());
+				const int32 MaxFrameIndexVisibility = MinFrameIndexVisibility + NumSamplesVisibility - 1;
+
+				MinTime = FMath::Min(MinTime, StartTimeVisibility);
+				MaxTime = FMath::Max(MaxTime, (float)TimeSampler->getSampleTime(NumSamplesVisibility - 1));
+				NumFrames = FMath::Max(NumFrames, NumSamplesVisibility);
+				MinFrameIndex = FMath::Min(MinFrameIndex, MinFrameIndexVisibility);
+				MaxFrameIndex = FMath::Max(MaxFrameIndex, MaxFrameIndexVisibility);
+			}
+		}
+	}
+
 	if (RootObject == nullptr && CreatedObject != nullptr)
 	{
 		RootObject = CreatedObject;
