PR #10722: GeometryCache fixes
* Fixed geometry cache material mapping when visibility animation is used. Now, the materials won't shift around when a mesh section is hidden.
* Import can now be interrupted.

diff --git a/Engine/Plugins/Importers/AlembicImporter/Source/AlembicLibrary/Private/AbcFile.cpp b/Engine/Plugins/Importers/AlembicImporter/Source/AlembicLibrary/Private/AbcFile.cpp
--- a/Engine/Plugins/Importers/AlembicImporter/Source/AlembicLibrary/Private/AbcFile.cpp
+++ b/Engine/Plugins/Importers/AlembicImporter/Source/AlembicLibrary/Private/AbcFile.cpp
@@ -18,6 +18,7 @@
 #include "AssetRegistry/AssetRegistryModule.h"
 #include "Logging/TokenizedMessage.h"
 #include "AbcImportLogger.h"
+#include "Misc/ScopedSlowTask.h"
 
 
 
@@ -37,6 +38,8 @@ THIRD_PARTY_INCLUDES_END
 
 #define LOCTEXT_NAMESPACE "AbcFile"
 
+const FString FAbcFile::NoFaceSetNameStr(TEXT("NoFaceSetName"));
+
 FAbcFile::FAbcFile(const FString& InFilePath)
 	: FilePath(InFilePath)
 	, RootObject(nullptr)
@@ -377,24 +380,47 @@ EAbcImportError FAbcFile::Import(UAbcImportSettings* InImportSettings)
 		}
 	}
 
+	// Count total number of facesets (1 or more per mesh)
+	int32 NumFacesets = 0;
+	for (const FAbcPolyMesh* PolyMesh : PolyMeshes)
+	{
+		if (PolyMesh->bShouldImport)
+		{
+			check(!PolyMesh->FaceSetNames.IsEmpty());
+			NumFacesets += PolyMesh->FaceSetNames.Num();
+		}
+	}
+	
+	if (NumFacesets == 0)
+	{
+		return AbcImportError_NoMeshes;
+	}
+
 	// Populate the list of unique face set names from the meshes that should be imported regardless of the import material settings
-	bool bRequiresDefaultMaterial = false;
-	for (FAbcPolyMesh* PolyMesh : PolyMeshes)
+	// and construct mapping from faceset index (in flat list of imported mesh facesets) to material slot
+	LookupMaterialSlot.SetNum(NumFacesets);
+
+	int32 FacesetMaterialIndex = 0;
+	for (const FAbcPolyMesh* PolyMesh : PolyMeshes)
 	{
 		if (PolyMesh->bShouldImport)
 		{
 			for (const FString& FaceSetName : PolyMesh->FaceSetNames)
 			{
-				UniqueFaceSetNames.AddUnique(FaceSetName);
+				int32 MaterialSlotIndex = UniqueFaceSetNames.Find(FaceSetName);
+				if (MaterialSlotIndex == INDEX_NONE)
+				{
+					MaterialSlotIndex = UniqueFaceSetNames.Num();
+					UniqueFaceSetNames.Add(FaceSetName);
+				}
+
+				LookupMaterialSlot[FacesetMaterialIndex] = MaterialSlotIndex;
+				FacesetMaterialIndex++;
 			}
-			bRequiresDefaultMaterial |= PolyMesh->FaceSetNames.Num() == 0;
 		}
 	}
 
-	if (bRequiresDefaultMaterial)
-	{
-		UniqueFaceSetNames.Insert(TEXT("DefaultMaterial"), 0);
-	}
+	check(FacesetMaterialIndex == NumFacesets);
 
 	return AbcImportError_NoError;
 }
@@ -541,17 +567,24 @@ void FAbcFile::CleanupFrameData(const int32 ReadIndex)
 	}
 }
 
-void FAbcFile::ProcessFrames(TFunctionRef<void(int32, FAbcFile*)> InCallback, const EFrameReadFlags InFlags)
+bool FAbcFile::ProcessFrames(TFunctionRef<void(int32, FAbcFile*)> InCallback, const EFrameReadFlags InFlags, const FSlowTask* SlowTask)
 {
 	const int32 NumWorkerThreads = FMath::Clamp(FTaskGraphInterface::Get().GetNumWorkerThreads(), 1, MaxNumberOfResidentSamples);
 	const bool bSingleThreaded = ImportSettings->NumThreads == 1 ||
 								 EnumHasAnyFlags(InFlags, EFrameReadFlags::ForceSingleThreaded) || !FApp::ShouldUseThreadingForPerformance();
 
+	volatile bool bIsCancelled = false;
+
 	if (bSingleThreaded)
 	{
 		for (int32 FrameIndex = StartFrameIndex; FrameIndex <= EndFrameIndex; ++FrameIndex)
 		{
-			ReadFrame(FrameIndex, InFlags, 0);
+			if (SlowTask && SlowTask->ShouldCancel())
+			{
+				return false;
+			}
+
+			ReadFrame(FrameIndex, InFlags, 0);		
 			InCallback(FrameIndex, this);
 			CleanupFrameData(0);
 		}
@@ -566,15 +599,21 @@ void FAbcFile::ProcessFrames(TFunctionRef<void(int32, FAbcFile*)> InCallback, co
 		{
 			TSharedRef<FTokenizedMessage> Message = FTokenizedMessage::Create(EMessageSeverity::Error, LOCTEXT("NoSynchEvent", "Unable to get synchronization event for parallelized Alembic frame data import."));
 			FAbcImportLogger::AddImportMessage(Message);
-			return;
+			return false;
 		}
 
-		ParallelFor(NumWorkerThreads, [this, InFlags, InCallback, bSingleThreaded, NumWorkerThreads, &WriteFrameIndex, &Mutex, &FrameWrittenEvent](int32 ThreadIndex)
+		ParallelFor(NumWorkerThreads, [this, InFlags, InCallback, bSingleThreaded, NumWorkerThreads, &WriteFrameIndex, &Mutex, &FrameWrittenEvent, &bIsCancelled, SlowTask](int32 ThreadIndex)
 		{
 			int32 FrameIndex = StartFrameIndex + ThreadIndex;
 
-			while (FrameIndex <= EndFrameIndex)
+			while (FrameIndex <= EndFrameIndex && !bIsCancelled)
 			{
+				if (IsInGameThread() && SlowTask && SlowTask->ShouldCancel())
+				{
+					bIsCancelled = true;
+					break;
+				}
+
 				// Read frame data into memory
 				ReadFrame(FrameIndex, InFlags, ThreadIndex);
 
@@ -606,6 +645,8 @@ void FAbcFile::ProcessFrames(TFunctionRef<void(int32, FAbcFile*)> InCallback, co
 
 		FPlatformProcess::ReturnSynchEventToPool(FrameWrittenEvent);
 	}
+
+	return !bIsCancelled;
 }
 
 const int32 FAbcFile::GetMinFrameIndex() const
diff --git a/Engine/Plugins/Importers/AlembicImporter/Source/AlembicLibrary/Private/AbcImportUtilities.cpp b/Engine/Plugins/Importers/AlembicImporter/Source/AlembicLibrary/Private/AbcImportUtilities.cpp
--- a/Engine/Plugins/Importers/AlembicImporter/Source/AlembicLibrary/Private/AbcImportUtilities.cpp
+++ b/Engine/Plugins/Importers/AlembicImporter/Source/AlembicLibrary/Private/AbcImportUtilities.cpp
@@ -57,36 +57,41 @@ uint32 AbcImporterUtilities::GenerateMaterialIndicesFromFaceSets(Alembic::AbcGeo
 	std::vector<std::string> FaceSetNames;
 	Schema.getFaceSetNames(FaceSetNames);
 
+	if (FaceSetNames.size() == 0)
+	{
+		return 1; // We will create a default unnamed faceset
+	}
+
 	// Number of unique face sets found in the Alembic Object
 	uint32 NumUniqueFaceSets = 0;
-	if (FaceSetNames.size() != 0)
+	
+	// Loop over the face-set names
+	for (uint32 FaceSetIndex = 0; FaceSetIndex < FaceSetNames.size(); ++FaceSetIndex)
 	{
-		// Loop over the face-set names
-		for (uint32 FaceSetIndex = 0; FaceSetIndex < FaceSetNames.size(); ++FaceSetIndex)
-		{
-			Alembic::AbcGeom::IFaceSet FaceSet = Schema.getFaceSet(FaceSetNames[FaceSetIndex]);
-			Alembic::AbcGeom::IFaceSetSchema FaceSetSchema = FaceSet.getSchema();
-			Alembic::AbcGeom::IFaceSetSchema::Sample FaceSetSample;
-			FaceSetSchema.get(FaceSetSample, FrameSelector);
+		const std::string& FaceSetName = FaceSetNames[FaceSetIndex];
 
-			// Retrieve face indices that are part of this face set
-			Alembic::Abc::Int32ArraySamplePtr Faces = FaceSetSample.getFaces();
-			const bool bFacesAvailable = (Faces != nullptr);
-			const int NumFaces = Faces->size();
+		Alembic::AbcGeom::IFaceSet FaceSet = Schema.getFaceSet(FaceSetName);
+		Alembic::AbcGeom::IFaceSetSchema FaceSetSchema = FaceSet.getSchema();
+		Alembic::AbcGeom::IFaceSetSchema::Sample FaceSetSample;
+		FaceSetSchema.get(FaceSetSample, FrameSelector);
 
-			// Set the shared Material index for all the contained faces
-			for (int32 i = 0; i < NumFaces && NumFaces < MaterialIndicesOut.Num(); ++i)
+		// Retrieve face indices that are part of this face set
+		Alembic::Abc::Int32ArraySamplePtr Faces = FaceSetSample.getFaces();
+		const bool bFacesAvailable = (Faces != nullptr);
+		const int NumFaces = Faces->size();
+
+		// Set the shared Material index for all the contained faces
+		for (int32 i = 0; i < NumFaces && NumFaces < MaterialIndicesOut.Num(); ++i)
+		{
+			const int32 FaceIndex = Faces->get()[i];
+			if (MaterialIndicesOut.IsValidIndex(FaceIndex))
 			{
-				const int32 FaceIndex = Faces->get()[i];
-				if (MaterialIndicesOut.IsValidIndex(FaceIndex))
-				{
-					MaterialIndicesOut[FaceIndex] = FaceSetIndex;
-				}
+				MaterialIndicesOut[FaceIndex] = FaceSetIndex;
 			}
-
-			// Found a new unique faceset
-			NumUniqueFaceSets++;
 		}
+
+		// Found a new unique faceset
+		NumUniqueFaceSets++;
 	}
 
 	return NumUniqueFaceSets;
@@ -102,6 +107,11 @@ void AbcImporterUtilities::RetrieveFaceSetNames(Alembic::AbcGeom::IPolyMeshSchem
 	{
 		NamesOut.Add(FString(Name.c_str()));
 	}
+
+	if (NamesOut.IsEmpty()) 
+	{
+		NamesOut.Add(FAbcFile::NoFaceSetNameStr);
+	}
 }
 
 void AbcImporterUtilities::TriangulateIndexBuffer(const TArray<uint32>& InFaceCounts, TArray<uint32>& InOutIndices)
@@ -262,7 +272,14 @@ ESampleReadFlags AbcImporterUtilities::GenerateAbcMeshSampleReadFlags(const Alem
 		bool bConstantFaceSets = true;
 		for (int32 FaceSetIndex = 0; FaceSetIndex < FaceSetNames.size(); ++FaceSetIndex)
 		{
-			Alembic::AbcGeom::IFaceSet FaceSet = MutableSchema->getFaceSet(FaceSetNames[FaceSetIndex]);
+			const std::string& FaceSetName = FaceSetNames[FaceSetIndex];
+
+			if (!MutableSchema->hasFaceSet(FaceSetName))
+			{
+				continue;
+			}
+
+			Alembic::AbcGeom::IFaceSet FaceSet = MutableSchema->getFaceSet(FaceSetName);
 			Alembic::AbcGeom::IFaceSetSchema FaceSetSchema = FaceSet.getSchema();
 			bConstantFaceSets &= FaceSetSchema.isConstant();
 		}
@@ -662,24 +679,14 @@ bool AbcImporterUtilities::GenerateAbcMeshSampleDataForFrame(const Alembic::AbcG
 		}
 	}
 
-	if (EnumHasAnyFlags(ReadFlags, ESampleReadFlags::MaterialIndices))
-	{
-		// Pre initialize face-material indices
-		Sample->MaterialIndices.AddZeroed(Sample->Indices.Num() / 3);
-		Sample->NumMaterials = GenerateMaterialIndicesFromFaceSets(*const_cast<Alembic::AbcGeom::IPolyMeshSchema*>(&Schema), FrameSelector, Sample->MaterialIndices);
+	// Pre initialize face-material indices
+	Sample->MaterialIndices.AddZeroed(Sample->Indices.Num() / 3);
+	Sample->NumMaterials = GenerateMaterialIndicesFromFaceSets(*const_cast<Alembic::AbcGeom::IPolyMeshSchema*>(&Schema), FrameSelector, Sample->MaterialIndices);
 
-		// Triangulate material face indices if needed
-		if (bNeedsTriangulation)
-		{
-			TriangulateMaterialIndices(FaceCounts, Sample->MaterialIndices);
-		}
-	}
-	else
+	// Triangulate material face indices if needed
+	if (bNeedsTriangulation)
 	{
-		if (Sample->MaterialIndices.Num() < ( Sample->Indices.Num() / 3))
-		{
-			Sample->MaterialIndices.AddZeroed((Sample->Indices.Num() / 3) - Sample->MaterialIndices.Num());
-		}
+		TriangulateMaterialIndices(FaceCounts, Sample->MaterialIndices);
 	}
 
 	return bRetrievalResult;
@@ -1537,8 +1544,6 @@ void AbcImporterUtilities::GeometryCacheDataForMeshSample(FGeometryCacheMeshData
 	OutMeshData.VertexInfo.bHasMotionVectors = bHasVelocities;
 	OutMeshData.VertexInfo.bHasImportedVertexNumbers = bHasImportedVertexNumbers;
 
-	uint32 NumMaterials = MaterialOffset;
-
 	const int32 NumTriangles = MeshSample->Indices.Num() / 3;
 	const uint32 NumSections = MeshSample->NumMaterials ? MeshSample->NumMaterials : 1;
 
@@ -1609,8 +1614,10 @@ void AbcImporterUtilities::GeometryCacheDataForMeshSample(FGeometryCacheMeshData
 		}
 	}
 
+	uint32 NumMaterials = MaterialOffset;
+
 	TArray<uint32>& Indices = OutMeshData.Indices;
-	for (uint32 SectionIndex = 0; SectionIndex < NumSections; ++SectionIndex)
+	for (uint32 SectionIndex = 0; SectionIndex < NumSections; ++SectionIndex, ++NumMaterials)
 	{
 		// Sometimes empty sections seem to be in the file, filter these out
 		// as empty batches are not allowed by the geometry cache (They ultimately trigger checks in the renderer)
@@ -1623,8 +1630,6 @@ void AbcImporterUtilities::GeometryCacheDataForMeshSample(FGeometryCacheMeshData
 		FGeometryCacheMeshBatchInfo BatchInfo;
 		BatchInfo.StartIndex = Indices.Num();
 		BatchInfo.MaterialIndex = NumMaterials;
-		NumMaterials++;
-
 
 		BatchInfo.NumTriangles = SectionIndices[SectionIndex].Num() / 3;
 		Indices.Append(SectionIndices[SectionIndex]);
@@ -1633,7 +1638,9 @@ void AbcImporterUtilities::GeometryCacheDataForMeshSample(FGeometryCacheMeshData
 }
 
 void AbcImporterUtilities::MergePolyMeshesToMeshData(int32 FrameIndex, int32 FrameStart, float SecondsPerFrame, bool bUseVelocitiesAsMotionVectors,
-	const TArray<FAbcPolyMesh*>& PolyMeshes, const TArray<FString>& UniqueFaceSetNames, float& FrameTime,
+	const TArray<FAbcPolyMesh*>& PolyMeshes, 
+	const TArray<int32> LookupMaterialSlot,
+	float& FrameTime,
 	FGeometryCacheMeshData& MeshData, int32& PreviousNumVertices, bool& bConstantTopology, bool bStoreImportedVertexNumbers)
 {
 	FAbcMeshSample MergedSample;
@@ -1643,34 +1650,30 @@ void AbcImporterUtilities::MergePolyMeshesToMeshData(int32 FrameIndex, int32 Fra
 		if (PolyMesh->bShouldImport)
 		{
 			const int32 Offset = MergedSample.MaterialIndices.Num();
-			const int32 MaterialIndexOffset = MergedSample.NumMaterials;
 			bConstantTopology = bConstantTopology && PolyMesh->bConstantTopology;
+
+			// The caller expects this value to be calculated regardless of mesh visibility
+			FrameTime = PolyMesh->GetTimeForFrameIndex(FrameIndex);
+
 			if (PolyMesh->GetVisibility(FrameIndex))
 			{
 				const FAbcMeshSample* Sample = PolyMesh->GetSample(FrameIndex);
-				FrameTime = PolyMesh->GetTimeForFrameIndex(FrameIndex);
 				AbcImporterUtilities::AppendMeshSample(&MergedSample, Sample);
-				if (PolyMesh->FaceSetNames.Num() == 0)
-				{
-					FMemory::Memzero(MergedSample.MaterialIndices.GetData() + Offset, (MergedSample.MaterialIndices.Num() - Offset) * sizeof(int32));
-				}
-				else
+			
+				for (int32 Index = Offset; Index < MergedSample.MaterialIndices.Num(); ++Index)
 				{
-					for (int32 Index = Offset; Index < MergedSample.MaterialIndices.Num(); ++Index)
-					{
-						int32& MaterialIndex = MergedSample.MaterialIndices[Index];
-						if (PolyMesh->FaceSetNames.IsValidIndex(MaterialIndex - MaterialIndexOffset))
-						{
-							int32 FaceSetMaterialIndex = UniqueFaceSetNames.IndexOfByKey(PolyMesh->FaceSetNames[MaterialIndex - MaterialIndexOffset]);
-							MaterialIndex = FaceSetMaterialIndex != INDEX_NONE ? FaceSetMaterialIndex : 0;
-						}
-						else
-						{
-							MaterialIndex = 0;
-						}
-					}
+					int32 MaterialID = MergedSample.MaterialIndices[Index];
+					MergedSample.MaterialIndices[Index] = LookupMaterialSlot[MaterialID];
 				}
 			}
+			else
+			{
+				MergedSample.NumMaterials += PolyMesh->FaceSetNames.Num();
+
+				// If this mesh is invisible on this frame then the topology is not constant!
+				bConstantTopology = false;
+			}
+
 		}
 	}
 
@@ -1680,7 +1683,6 @@ void AbcImporterUtilities::MergePolyMeshesToMeshData(int32 FrameIndex, int32 Fra
 	}
 	PreviousNumVertices = MergedSample.Vertices.Num();
 
-	MergedSample.NumMaterials = UniqueFaceSetNames.Num();
 
 	// Generate the mesh data for this sample
 	AbcImporterUtilities::GeometryCacheDataForMeshSample(MeshData, &MergedSample, 0, SecondsPerFrame, bUseVelocitiesAsMotionVectors, bStoreImportedVertexNumbers);
diff --git a/Engine/Plugins/Importers/AlembicImporter/Source/AlembicLibrary/Private/AbcImportUtilities.h b/Engine/Plugins/Importers/AlembicImporter/Source/AlembicLibrary/Private/AbcImportUtilities.h
--- a/Engine/Plugins/Importers/AlembicImporter/Source/AlembicLibrary/Private/AbcImportUtilities.h
+++ b/Engine/Plugins/Importers/AlembicImporter/Source/AlembicLibrary/Private/AbcImportUtilities.h
@@ -337,7 +337,7 @@ namespace AbcImporterUtilities
 	 * @param SecondsPerFrame		The time step of the Abc file
 	 * @param bUseVelocitiesAsMotionVectors		Converts the AbcPolyMesh velocities to motion vectors
 	 * @param PolyMeshes			The PolyMeshes to merge, which will be sampled at FrameIndex
-	 * @param UniqueFaceSetNames	The array of unique face set names of the PolyMeshes
+	 * @param LookupMaterialSlot	Mapping from original face set index (index in flattened list of facesets to import) to material slot (one for each unique faceset name)
 	 * @param FrameTime				The frame time that corresponds to FrameIndex
 	 * @param MeshData				The GeometryCacheMeshData where to output the merged PolyMeshes
 	 * @param PreviousNumVertices	The number of vertices in the merged PolyMeshes, used to determine if its topology is constant between 2 frames
@@ -345,7 +345,9 @@ namespace AbcImporterUtilities
 	 * @param bStoreImportedVertexNumbers Set to true when we want to store the original dcc vertex numbers for each vertex.
 	 */
 	void MergePolyMeshesToMeshData(int32 FrameIndex, int32 FrameStart, float SecondsPerFrame, bool bUseVelocitiesAsMotionVectors,
-		const TArray<FAbcPolyMesh*>& PolyMeshes, const TArray<FString>& UniqueFaceSetNames, float& FrameTime,
+		const TArray<FAbcPolyMesh*>& PolyMeshes, 
+		const TArray<int32> LookupMaterialSlot,
+		float& FrameTime,
 		FGeometryCacheMeshData& MeshData, int32& PreviousNumVertices, bool& bConstantTopology, bool bStoreImportedVertexNumbers);
 
 	/** Retrieves a material from an AbcFile according to the given name and resaves it into the parent package */
diff --git a/Engine/Plugins/Importers/AlembicImporter/Source/AlembicLibrary/Private/AbcImporter.cpp b/Engine/Plugins/Importers/AlembicImporter/Source/AlembicLibrary/Private/AbcImporter.cpp
--- a/Engine/Plugins/Importers/AlembicImporter/Source/AlembicLibrary/Private/AbcImporter.cpp
+++ b/Engine/Plugins/Importers/AlembicImporter/Source/AlembicLibrary/Private/AbcImporter.cpp
@@ -372,7 +372,22 @@ UGeometryCache* FAbcImporter::ImportAsGeometryCache(UObject* InParent, EObjectFl
 		// Load the default material for later usage
 		UMaterial* DefaultMaterial = UMaterial::GetDefaultMaterial(MD_Surface);
 		check(DefaultMaterial);
-		uint32 MaterialOffset = 0;
+
+		auto CreateMaterial = [this, GeometryCache, DefaultMaterial, InParent, Flags](const FString& FaceSetName)
+		{
+			FName MaterialName(*FaceSetName);
+
+			UMaterialInterface* Material = AbcImporterUtilities::RetrieveMaterial(*AbcFile, FaceSetName, InParent, Flags);
+			check(Material);
+
+			if (Material != DefaultMaterial)
+			{
+				Material->PostEditChange();
+			}
+
+			GeometryCache->Materials.Add(Material);
+			GeometryCache->MaterialSlotNames.Add(MaterialName);
+		};
 
 		// Add tracks
 		const int32 NumPolyMeshes = AbcFile->GetNumPolyMeshes();
@@ -380,9 +395,6 @@ UGeometryCache* FAbcImporter::ImportAsGeometryCache(UObject* InParent, EObjectFl
 		{
 			TArray<UGeometryCacheTrackStreamable*> Tracks;
 
-			TArray<FAbcPolyMesh*> ImportPolyMeshes;
-			TArray<int32> MaterialOffsets;
-
 			const bool bContainsHeterogeneousMeshes = AbcFile->ContainsHeterogeneousMeshes();
 			if (ImportSettings->GeometryCacheSettings.bApplyConstantTopologyOptimizations && bContainsHeterogeneousMeshes)
 			{
@@ -390,6 +402,10 @@ UGeometryCache* FAbcImporter::ImportAsGeometryCache(UObject* InParent, EObjectFl
 				FAbcImportLogger::AddImportMessage(Message);
 			}
 
+			FScopedSlowTask SlowTask(static_cast<float>((ImportSettings->SamplingSettings.FrameEnd + 1) - ImportSettings->SamplingSettings.FrameStart), FText::FromString(FString(TEXT("Importing Frames"))));
+			SlowTask.MakeDialog(true);
+			float CompletedFrames = 0.0f;
+
 			if (ImportSettings->GeometryCacheSettings.bFlattenTracks)
 			{
 				//UGeometryCacheCodecRaw* Codec = NewObject<UGeometryCacheCodecRaw>(GeometryCache, FName(*FString(TEXT("Flattened_Codec"))), RF_Public);
@@ -401,18 +417,13 @@ UGeometryCache* FAbcImporter::ImportAsGeometryCache(UObject* InParent, EObjectFl
 				Track->BeginCoding(Codec, ImportSettings->GeometryCacheSettings.bApplyConstantTopologyOptimizations && !bContainsHeterogeneousMeshes, bCalculateMotionVectors, ImportSettings->GeometryCacheSettings.bOptimizeIndexBuffers);
 				Tracks.Add(Track);
 				
-				FScopedSlowTask SlowTask(static_cast<float>((ImportSettings->SamplingSettings.FrameEnd + 1) - ImportSettings->SamplingSettings.FrameStart), FText::FromString(FString(TEXT("Importing Frames"))));
-				SlowTask.MakeDialog();
-
-				const TArray<FString>& UniqueFaceSetNames = AbcFile->GetUniqueFaceSetNames();
 				const TArray<FAbcPolyMesh*>& PolyMeshes = AbcFile->GetPolyMeshes();
 				TArray<float> FrameTimes;
 				FrameTimes.SetNum(ImportSettings->SamplingSettings.FrameEnd - ImportSettings->SamplingSettings.FrameStart + 1);
 				
 				const int32 NumTracks = Tracks.Num();
 				int32 PreviousNumVertices = 0;
-				float CompletedFrames = 0.0f;
-				TFunction<void(int32, FAbcFile*)> Callback = [this, &Tracks, &SlowTask, &UniqueFaceSetNames, &PolyMeshes, &PreviousNumVertices, &FrameTimes, &CompletedFrames](int32 FrameIndex, const FAbcFile* InAbcFile)
+				TFunction<void(int32, FAbcFile*)> Callback = [this, &Tracks, &PolyMeshes, &PreviousNumVertices, &FrameTimes, &SlowTask, &CompletedFrames](int32 FrameIndex, const FAbcFile* InAbcFile)
 				{
 					const bool bUseVelocitiesAsMotionVectors = (ImportSettings->GeometryCacheSettings.MotionVectors == EAbcGeometryCacheMotionVectorsImport::ImportAbcVelocitiesAsMotionVectors);
 					FGeometryCacheMeshData MeshData;
@@ -421,7 +432,10 @@ UGeometryCache* FAbcImporter::ImportAsGeometryCache(UObject* InParent, EObjectFl
 
 					int32 FrameTimeIndex = FrameIndex - ImportSettings->SamplingSettings.FrameStart;
 					AbcImporterUtilities::MergePolyMeshesToMeshData(FrameIndex, ImportSettings->SamplingSettings.FrameStart, AbcFile->GetSecondsPerFrame(), bUseVelocitiesAsMotionVectors,
-						PolyMeshes, UniqueFaceSetNames, FrameTimes[FrameTimeIndex], MeshData, PreviousNumVertices, bConstantTopology, bStoreImportedVertexNumbers);
+						PolyMeshes,
+						AbcFile->GetLookupMaterialSlot(),
+						FrameTimes[FrameTimeIndex],
+						MeshData, PreviousNumVertices, bConstantTopology, bStoreImportedVertexNumbers);
 					
 					const float FrameRate = InAbcFile->GetFramerate();
 
@@ -437,23 +451,23 @@ UGeometryCache* FAbcImporter::ImportAsGeometryCache(UObject* InParent, EObjectFl
 					}					
 				};
 
-				AbcFile->ProcessFrames(Callback, EFrameReadFlags::ApplyMatrix);
-
-				// Now add materials for all the face set names
-				for (const FString& FaceSetName : UniqueFaceSetNames)
+				if (!AbcFile->ProcessFrames(Callback, EFrameReadFlags::ApplyMatrix, &SlowTask))
 				{
-					UMaterialInterface* Material = AbcImporterUtilities::RetrieveMaterial(*AbcFile, FaceSetName, InParent, Flags);
-					GeometryCache->Materials.Add((Material != nullptr) ? Material : DefaultMaterial);		
-					GeometryCache->MaterialSlotNames.Add(FaceSetName != TEXT("DefaultMaterial") ? FName(FaceSetName) : NoFaceSetName);
+					UE_LOG(LogAbcImporter, Warning, TEXT("Alembic geometry cache import was interrupted"));
+				}
 
-					if (Material != UMaterial::GetDefaultMaterial(MD_Surface))
-					{
-						Material->PostEditChange();
-					}
+				// Now add materials for all the unique face set names
+				for (const FString& FaceSetName : AbcFile->GetUniqueFaceSetNames())
+				{
+					CreateMaterial(FaceSetName);
 				}
 			}
 			else
 			{
+				uint32 MaterialOffset = 0;
+				TArray<int32> MaterialOffsets;
+				TArray<FAbcPolyMesh*> ImportPolyMeshes;
+
 				const TArray<FAbcPolyMesh*>& PolyMeshes = AbcFile->GetPolyMeshes();
 				for (FAbcPolyMesh* PolyMesh : PolyMeshes)
 				{
@@ -476,36 +490,17 @@ UGeometryCache* FAbcImporter::ImportAsGeometryCache(UObject* InParent, EObjectFl
 						MaterialOffsets.Add(MaterialOffset);
 
 						// Add materials for this Mesh Object
-						const uint32 NumMaterials = (PolyMesh->FaceSetNames.Num() > 0) ? PolyMesh->FaceSetNames.Num() : 1;
+						const uint32 NumMaterials = PolyMesh->FaceSetNames.Num();
 						for (uint32 MaterialIndex = 0; MaterialIndex < NumMaterials; ++MaterialIndex)
 						{
-							UMaterialInterface* Material = nullptr;
-							if (PolyMesh->FaceSetNames.IsValidIndex(MaterialIndex))
-							{
-								Material = AbcImporterUtilities::RetrieveMaterial(*AbcFile, PolyMesh->FaceSetNames[MaterialIndex], InParent, Flags);
-								GeometryCache->MaterialSlotNames.Add(FName(PolyMesh->FaceSetNames[MaterialIndex]));
-								if (Material != UMaterial::GetDefaultMaterial(MD_Surface))
-								{
-									Material->PostEditChange();
-								}
-							}
-							else
-							{
-								GeometryCache->MaterialSlotNames.Add(NoFaceSetName);
-							}
-
-							GeometryCache->Materials.Add((Material != nullptr) ? Material : DefaultMaterial);
+							CreateMaterial(PolyMesh->FaceSetNames[MaterialIndex]);
 						}
 
 						MaterialOffset += NumMaterials;
 					}
 				}
 
-				FScopedSlowTask SlowTask(static_cast<float>((ImportSettings->SamplingSettings.FrameEnd + 1) - ImportSettings->SamplingSettings.FrameStart), FText::FromString(FString(TEXT("Importing Frames"))));
-				SlowTask.MakeDialog();
-
 				const int32 NumTracks = Tracks.Num();
-				float CompletedFrames = 0.0f;
 				TFunction<void(int32, FAbcFile*)> Callback = [this, NumTracks, &ImportPolyMeshes, &Tracks, &MaterialOffsets, &SlowTask, &CompletedFrames](int32 FrameIndex, const FAbcFile* InAbcFile)
 				{
 					const float FrameRate = static_cast<float>(InAbcFile->GetFramerate());
@@ -533,7 +528,7 @@ UGeometryCache* FAbcImporter::ImportAsGeometryCache(UObject* InParent, EObjectFl
 							Track->AddVisibilitySample(bVisible, FrameTime);
 						}
 					}
-					
+
 					++CompletedFrames;
 					if (IsInGameThread())
 					{
@@ -542,7 +537,10 @@ UGeometryCache* FAbcImporter::ImportAsGeometryCache(UObject* InParent, EObjectFl
 					}					
 				};
 
-				AbcFile->ProcessFrames(Callback, EFrameReadFlags::ApplyMatrix);
+				if (!AbcFile->ProcessFrames(Callback, EFrameReadFlags::ApplyMatrix, &SlowTask))
+				{
+					UE_LOG(LogAbcImporter, Warning, TEXT("Alembic geometry cache import was interrupted"));
+				}
 			}
 
 			TArray<FMatrix> Mats;
@@ -930,7 +928,7 @@ const bool FAbcImporter::CompressAnimationDataUsingPCA(const FAbcCompressionSett
 			int32 NumSamples = 0;
 
 			TUniquePtr<FScopedSlowTask> SlowTask = MakeUnique<FScopedSlowTask>(static_cast<float>((ImportSettings->SamplingSettings.FrameEnd + 1) - ImportSettings->SamplingSettings.FrameStart + 1), FText::FromString(FString(TEXT("Merging meshes"))));
-			SlowTask->MakeDialog();
+			SlowTask->MakeDialog(true);
 
 			TArray<uint32> ObjectVertexOffsets;
 			TArray<uint32> ObjectIndexOffsets;
@@ -1027,8 +1025,8 @@ const bool FAbcImporter::CompressAnimationDataUsingPCA(const FAbcCompressionSett
 					const EAppReturnType::Type DialogResponse = FMessageDialog::Open(EAppMsgType::OkCancel, EAppReturnType::Ok, Message, Title);
 					if (DialogResponse != EAppReturnType::Ok)
 					{
-						return false;
-					}
+					return false;
+				}
 				}
 
 				AverageVertexData.Reset();
@@ -1041,7 +1039,11 @@ const bool FAbcImporter::CompressAnimationDataUsingPCA(const FAbcCompressionSett
 				NumSamples = 0;
 			}
 
-			AbcFile->ProcessFrames(MergedMeshesFunc, Flags);
+			if (!AbcFile->ProcessFrames(MergedMeshesFunc, Flags, SlowTask.Get()))
+			{
+				UE_LOG(LogAbcImporter, Warning, TEXT("Alembic skeletal mesh import was interrupted"));
+				return false;
+			}
 			SlowTask.Reset();
 
 			// Average out vertex data
@@ -1098,7 +1100,7 @@ const bool FAbcImporter::CompressAnimationDataUsingPCA(const FAbcCompressionSett
 			}
 
 			SlowTask = MakeUnique<FScopedSlowTask>(static_cast<float>((ImportSettings->SamplingSettings.FrameEnd + 1) - ImportSettings->SamplingSettings.FrameStart), FText::FromString(FString(TEXT("Generating matrices"))));
-			SlowTask->MakeDialog();
+			SlowTask->MakeDialog(true);
 
 			CompletedFrames = 0.0f;
 
@@ -1149,7 +1151,11 @@ const bool FAbcImporter::CompressAnimationDataUsingPCA(const FAbcCompressionSett
 					}					
 				};
 
-			AbcFile->ProcessFrames(GenerateMatrixFunc, Flags);
+			if (!AbcFile->ProcessFrames(GenerateMatrixFunc, Flags, SlowTask.Get()))
+			{
+				UE_LOG(LogAbcImporter, Warning, TEXT("Alembic skeletal mesh import was interrupted"));
+				return false;
+			}
 			SlowTask.Reset();
 
 			// Perform compression
@@ -1217,8 +1223,8 @@ const bool FAbcImporter::CompressAnimationDataUsingPCA(const FAbcCompressionSett
 			AverageNormalData.AddDefaulted(NumPolyMeshesToCompress);
 			
 			TUniquePtr<FScopedSlowTask> SlowTask = MakeUnique<FScopedSlowTask>(static_cast<float>((ImportSettings->SamplingSettings.FrameEnd + 1) - ImportSettings->SamplingSettings.FrameStart + 1), FText::FromString(FString(TEXT("Processing meshes"))));
-			SlowTask->MakeDialog();
-
+			SlowTask->MakeDialog(true);
+			
 			int32 NumSamples = 0;
 			float CompletedFrames = 0.0f;
 			TFunction<void(int32, FAbcFile*)> IndividualMeshesFunc =
@@ -1345,7 +1351,11 @@ const bool FAbcImporter::CompressAnimationDataUsingPCA(const FAbcCompressionSett
 				NumSamples = 0;
 			}
 
-			AbcFile->ProcessFrames(IndividualMeshesFunc, Flags);
+			if (!AbcFile->ProcessFrames(IndividualMeshesFunc, Flags, SlowTask.Get()))
+			{
+				UE_LOG(LogAbcImporter, Warning, TEXT("Alembic skeletal mesh import was interrupted"));
+				return false;
+			}
 			SlowTask.Reset();
 
 			// Average out vertex data
@@ -1385,7 +1395,7 @@ const bool FAbcImporter::CompressAnimationDataUsingPCA(const FAbcCompressionSett
 			}
 
 			SlowTask = MakeUnique<FScopedSlowTask>(static_cast<float>((ImportSettings->SamplingSettings.FrameEnd + 1) - ImportSettings->SamplingSettings.FrameStart), FText::FromString(FString(TEXT("Generating matrices"))));
-			SlowTask->MakeDialog();
+			SlowTask->MakeDialog(true);
 
 			uint32 GenerateMatrixSampleIndex = 0;
 			CompletedFrames = 0.0f;
@@ -1439,9 +1449,13 @@ const bool FAbcImporter::CompressAnimationDataUsingPCA(const FAbcCompressionSett
 						SlowTask->EnterProgressFrame(CompletedFrames);
 						CompletedFrames = 0.0f;
 					}					
-			};
+				};
 
-			AbcFile->ProcessFrames(GenerateMatrixFunc, Flags);
+			if (!AbcFile->ProcessFrames(GenerateMatrixFunc, Flags, SlowTask.Get()))
+			{
+				UE_LOG(LogAbcImporter, Warning, TEXT("Alembic skeletal mesh import was interrupted"));
+				return false;
+			}
 			SlowTask.Reset();
 
 			for (int32 MeshIndex = 0; MeshIndex < NumPolyMeshesToCompress; ++MeshIndex)
diff --git a/Engine/Plugins/Importers/AlembicImporter/Source/AlembicLibrary/Private/AbcUtilities.cpp b/Engine/Plugins/Importers/AlembicImporter/Source/AlembicLibrary/Private/AbcUtilities.cpp
--- a/Engine/Plugins/Importers/AlembicImporter/Source/AlembicLibrary/Private/AbcUtilities.cpp
+++ b/Engine/Plugins/Importers/AlembicImporter/Source/AlembicLibrary/Private/AbcUtilities.cpp
@@ -26,7 +26,7 @@ void FAbcUtilities::GetFrameMeshData(FAbcFile& AbcFile, int32 FrameIndex, FGeome
 	const bool bStoreImportedVertexNumbers = ImportSettings->GeometryCacheSettings.bStoreImportedVertexNumbers;
 
 	AbcImporterUtilities::MergePolyMeshesToMeshData(FrameIndex, 0, AbcFile.GetSecondsPerFrame(), bUseVelocitiesAsMotionVectors,
-		AbcFile.GetPolyMeshes(), AbcFile.GetUniqueFaceSetNames(), FrameTime, MeshData, PreviousNumVertices, bConstantTopology, bStoreImportedVertexNumbers);
+		AbcFile.GetPolyMeshes(), AbcFile.GetLookupMaterialSlot(), FrameTime, MeshData, PreviousNumVertices, bConstantTopology, bStoreImportedVertexNumbers);
 
 	OutMeshData = MoveTemp(MeshData);
 
diff --git a/Engine/Plugins/Importers/AlembicImporter/Source/AlembicLibrary/Public/AbcFile.h b/Engine/Plugins/Importers/AlembicImporter/Source/AlembicLibrary/Public/AbcFile.h
--- a/Engine/Plugins/Importers/AlembicImporter/Source/AlembicLibrary/Public/AbcFile.h
+++ b/Engine/Plugins/Importers/AlembicImporter/Source/AlembicLibrary/Public/AbcFile.h
@@ -25,6 +25,7 @@ class FAbcTransform;
 class UAbcImportSettings;
 class UMaterialInterface;
 class IMeshUtilities;
+struct FSlowTask;
 
 enum EAbcImportError : uint32;
 
@@ -56,8 +57,8 @@ public:
 	/** Get file path for currently opened ABC file*/
 	const FString GetFilePath() const;
 
-	/** Process all to-import frames for each poly mesh inside of the ABC file */
-	void ProcessFrames(TFunctionRef<void(int32, FAbcFile*)> InCallback, const EFrameReadFlags InFlags);
+	/** Process all to-import frames for each poly mesh inside of the ABC file, optional SlowTask to interrupt process and returns 'completed' */
+	bool ProcessFrames(TFunctionRef<void(int32, FAbcFile*)> InCallback, const EFrameReadFlags InFlags, const FSlowTask* SlowTask = nullptr);
 	
 	/** Returns ABC specific frame and timing information */
 	const int32 GetMinFrameIndex() const; // index of the first frame with data regardless of the imported range
@@ -96,11 +97,17 @@ public:
 	void CleanupFrameData(const int32 ReadIndex);
 	/** Returns the list of unique face set names from the meshes to be imported */
 	const TArray<FString>& GetUniqueFaceSetNames() const { return UniqueFaceSetNames; }
+	/** Returns the mapping from material index (stored in per triangle material indices of FAbcMeshSample) to material slot (one for each unique faceset name) */
+	const TArray<int32>& GetLookupMaterialSlot() const { return LookupMaterialSlot; }
+	/** Returns the number of material slots ( == number of unique facesets ) */
+	int32 GetNumMaterialSlots() const { return GetUniqueFaceSetNames().Num(); }
 
 	typedef TPair<FString, FString> FMetaData;
 	/** Returns the metadata of the Alembic archive */
 	TArray<FMetaData> GetArchiveMetaData() const;
 
+	static const FString NoFaceSetNameStr;
+
 protected:
 	void TraverseAbcHierarchy(const Alembic::Abc::IObject& InObject, IAbcObject* InParent);
 	void ExtractCustomAttributes(const Alembic::AbcGeom::IPolyMesh& InMesh);
@@ -135,6 +142,7 @@ protected:
 	/** Map of material created for the imported alembic file identified by material names */
 	TMap<FString, UMaterialInterface*> MaterialMap;
 	TArray<FString> UniqueFaceSetNames;
+	TArray<int32> LookupMaterialSlot;
 
 	/** Total (max) number of frames in the Alembic file */
 	int32 NumFrames;
diff --git a/Engine/Plugins/Runtime/GeometryCache/Source/GeometryCache/Private/GeometryCacheSceneProxy.cpp b/Engine/Plugins/Runtime/GeometryCache/Source/GeometryCache/Private/GeometryCacheSceneProxy.cpp
--- a/Engine/Plugins/Runtime/GeometryCache/Source/GeometryCache/Private/GeometryCacheSceneProxy.cpp
+++ b/Engine/Plugins/Runtime/GeometryCache/Source/GeometryCache/Private/GeometryCacheSceneProxy.cpp
@@ -484,7 +484,8 @@ void FGeometryCacheSceneProxy::GetDynamicMeshElements(const TArray<const FSceneV
 #endif
 
 						// Apply view mode material overrides
-						FMaterialRenderProxy* MaterialProxy = bWireframe ? WireframeMaterialInstance : TrackProxy->Materials[BatchIndex]->GetRenderProxy();
+						const int32 MaterialIndex = TrackProxy->Materials.IsValidIndex(BatchInfo.MaterialIndex) ? BatchInfo.MaterialIndex : BatchIndex; 
+						FMaterialRenderProxy* MaterialProxy = bWireframe ? WireframeMaterialInstance : TrackProxy->Materials[MaterialIndex]->GetRenderProxy();
 						MeshBatch.bWireframe = bWireframe;
 						MeshBatch.MaterialRenderProxy = MaterialProxy;
 
@@ -536,7 +537,8 @@ void FGeometryCacheSceneProxy::GetDynamicRayTracingInstances(FRayTracingMaterial
 			FDynamicPrimitiveUniformBuffer& DynamicPrimitiveUniformBuffer = Context.RayTracingMeshResourceCollector.AllocateOneFrameResource<FDynamicPrimitiveUniformBuffer>();
 			CreateMeshBatch(RHICmdList, TrackProxy, BatchInfo, UserDataWrapper, DynamicPrimitiveUniformBuffer, MeshBatch);
 
-			MeshBatch.MaterialRenderProxy = TrackProxy->Materials[SegmentIndex]->GetRenderProxy();
+			const int32 MaterialIndex = TrackProxy->Materials.IsValidIndex(BatchInfo.MaterialIndex) ? BatchInfo.MaterialIndex : SegmentIndex;
+			MeshBatch.MaterialRenderProxy = TrackProxy->Materials[MaterialIndex]->GetRenderProxy();
 			MeshBatch.CastRayTracedShadow = IsShadowCast(Context.ReferenceView);
 			MeshBatch.SegmentIndex = SegmentIndex;
 
