PR #10723: AlembicImporter: static and skeletal mesh material slot names and slot merging. Each unique face set now translates to a unique material slot regardless of the "Merge meshes" setting.

diff --git a/Engine/Plugins/Importers/AlembicImporter/Source/AlembicLibrary/Private/AbcImporter.cpp b/Engine/Plugins/Importers/AlembicImporter/Source/AlembicLibrary/Private/AbcImporter.cpp
--- a/Engine/Plugins/Importers/AlembicImporter/Source/AlembicLibrary/Private/AbcImporter.cpp
+++ b/Engine/Plugins/Importers/AlembicImporter/Source/AlembicLibrary/Private/AbcImporter.cpp
@@ -184,7 +184,7 @@ T* FAbcImporter::CreateObjectInstance(UObject*& InParent, const FString& ObjectN
 	return NewObject<T>(Package, FName(*SanitizedObjectName), Flags | RF_Public);
 }
 
-UStaticMesh* FAbcImporter::CreateStaticMeshFromSample(UObject* InParent, const FString& Name, EObjectFlags Flags, const uint32 NumMaterials, const TArray<FString>& FaceSetNames, const FAbcMeshSample* Sample)
+UStaticMesh* FAbcImporter::CreateStaticMeshFromSample(UObject* InParent, const FString& Name, EObjectFlags Flags, const TArray<FString>& UniqueFaceSetNames, const TArray<int32>& LookupMaterialSlot, const FAbcMeshSample* Sample)
 {
 	UStaticMesh* StaticMesh = CreateObjectInstance<UStaticMesh>(InParent, Name, Flags);
 
@@ -208,35 +208,28 @@ UStaticMesh* FAbcImporter::CreateStaticMeshFromSample(UObject* InParent, const F
 
 		// Material list
 		StaticMesh->GetStaticMaterials().Empty();
-		// If there were FaceSets available in the Alembic file use the number of unique face sets as num material entries, otherwise default to one material for the whole mesh
-		const uint32 FrameIndex = 0;
-		uint32 NumFaceSets = FaceSetNames.Num();
 
-		const bool bCreateMaterial = ImportSettings->MaterialSettings.bCreateMaterials;
-		for (uint32 MaterialIndex = 0; MaterialIndex < ((NumMaterials != 0) ? NumMaterials : 1); ++MaterialIndex)
+		// Build the material slots : one for each faceset
+		for (const FString& FaceSetName : UniqueFaceSetNames)
 		{
-			UMaterialInterface* Material = nullptr;
-			if (FaceSetNames.IsValidIndex(MaterialIndex))
-			{
-				Material = AbcImporterUtilities::RetrieveMaterial(*AbcFile, FaceSetNames[MaterialIndex], InParent, Flags);
+			UMaterialInterface* Material = AbcImporterUtilities::RetrieveMaterial(*AbcFile, FaceSetName, InParent, Flags);
 
-				if (Material != UMaterial::GetDefaultMaterial(MD_Surface))
-				{
-					Material->PostEditChange();
-				}
+			if (Material != DefaultMaterial)
+			{
+				Material->PostEditChange();
 			}
 
-			FStaticMaterial StaticMaterial;
-			StaticMaterial.MaterialInterface = (Material != nullptr) ? Material : DefaultMaterial;
+			if (Material == nullptr)
+			{
+				Material = DefaultMaterial;
+			}
 
-			FName SlotName((Material != nullptr) ? FName(FaceSetNames[MaterialIndex]) : NoFaceSetName);
-			StaticMaterial.MaterialSlotName = SlotName;
-			StaticMaterial.ImportedMaterialSlotName = SlotName;
+			FName MaterialName(*FaceSetName);
 
-			StaticMesh->GetStaticMaterials().Add(StaticMaterial);
+			StaticMesh->GetStaticMaterials().Add(FStaticMaterial(Material, MaterialName, MaterialName));
 		}
 
-		GenerateMeshDescriptionFromSample(Sample, MeshDescription, StaticMesh);
+		GenerateMeshDescriptionFromSample(UniqueFaceSetNames, LookupMaterialSlot, Sample, MeshDescription);
 
 		// Get the first LOD for filling it up with geometry, only support one LOD
 		FStaticMeshSourceModel& SrcModel = StaticMesh->GetSourceModel(LODIndex);
@@ -277,46 +270,36 @@ const TArray<UStaticMesh*> FAbcImporter::ImportAsStaticMesh(UObject* InParent, E
 	TFunction<void(int32, FAbcFile*)> Func = [this, &ImportedStaticMeshes, StaticMeshSettings, InParent, Flags](int32 FrameIndex, FAbcFile* InFile)
 	{
 		const TArray<FAbcPolyMesh*>& PolyMeshes = AbcFile->GetPolyMeshes();
-		if (StaticMeshSettings.bMergeMeshes)
+		if (StaticMeshSettings.bMergeMeshes) 
 		{
-			// If merging we merge all the raw mesh structures together and generate a static mesh asset from this
-			TArray<FString> MergedFaceSetNames;
-			TArray<FAbcMeshSample*> Samples;
-			uint32 TotalNumMaterials = 0;
-
+			// Merge all meshes in the Alembic cache to one single static mesh
 			TArray<const FAbcMeshSample*> SamplesToMerge;
-			// Should merge all samples in the Alembic cache to one single static mesh
 			for (const FAbcPolyMesh* PolyMesh : PolyMeshes)
 			{
 				if (PolyMesh->bShouldImport)
 				{
 					const FAbcMeshSample* Sample = PolyMesh->GetSample(FrameIndex);
 					SamplesToMerge.Add(Sample);
-					TotalNumMaterials += (Sample->NumMaterials != 0) ? Sample->NumMaterials : 1;
-
-					if (PolyMesh->FaceSetNames.Num() > 0)
-					{
-						MergedFaceSetNames.Append(PolyMesh->FaceSetNames);
-					}
-					else
-					{
-						// Default name
-						MergedFaceSetNames.Add(NoFaceSetNameStr);
-					}
 				}
 			}
 
-			// Only merged samples if there are any
 			if (SamplesToMerge.Num())
 			{
 				FAbcMeshSample* MergedSample = AbcImporterUtilities::MergeMeshSamples(SamplesToMerge);
+			
+				UStaticMesh* StaticMesh = CreateStaticMeshFromSample(InParent, 
+					InParent != GetTransientPackage() ? FPaths::GetBaseFilename(InParent->GetName()) : (FPaths::GetBaseFilename(AbcFile->GetFilePath()) + "_" + FGuid::NewGuid().ToString()),
+					Flags, 
+					AbcFile->GetUniqueFaceSetNames(),
+					AbcFile->GetLookupMaterialSlot(),
+					MergedSample);
 
-
-				UStaticMesh* StaticMesh = CreateStaticMeshFromSample(InParent, InParent != GetTransientPackage() ? FPaths::GetBaseFilename(InParent->GetName()) : (FPaths::GetBaseFilename(AbcFile->GetFilePath()) + "_" + FGuid::NewGuid().ToString()), Flags, TotalNumMaterials, MergedFaceSetNames, MergedSample);
 				if (StaticMesh)
 				{
-						ImportedStaticMeshes.Add(StaticMesh);
+					ImportedStaticMeshes.Add(StaticMesh);
 				}
+
+				delete MergedSample; // Delete this temporary mesh
 			}
 		}
 		else
@@ -326,8 +309,19 @@ const TArray<UStaticMesh*> FAbcImporter::ImportAsStaticMesh(UObject* InParent, E
 				const FAbcMeshSample* Sample = PolyMesh->GetSample(FrameIndex);
 				if (PolyMesh->bShouldImport && Sample)
 				{
+					TArray<int32> LookupMaterialSlot;
+					for (int32 MaterialSlotIndex = 0; MaterialSlotIndex < PolyMesh->FaceSetNames.Num(); ++MaterialSlotIndex)
+					{
+						LookupMaterialSlot.Add(MaterialSlotIndex);
+					}
+
 					// Setup static mesh instance
-					UStaticMesh* StaticMesh = CreateStaticMeshFromSample(InParent, InParent != GetTransientPackage() ? PolyMesh->GetName() : PolyMesh->GetName() + "_" + FGuid::NewGuid().ToString(), Flags, Sample->NumMaterials, PolyMesh->FaceSetNames, Sample);
+					UStaticMesh* StaticMesh = CreateStaticMeshFromSample(InParent, 
+						InParent != GetTransientPackage() ? PolyMesh->GetName() : PolyMesh->GetName() + "_" + FGuid::NewGuid().ToString(), 
+						Flags, 
+						PolyMesh->FaceSetNames, 
+						LookupMaterialSlot,
+						Sample);
 
 					if (StaticMesh)
 					{
@@ -638,26 +632,33 @@ TArray<UObject*> FAbcImporter::ImportAsSkeletalMesh(UObject* InParent, EObjectFl
 		}
 
 
-		FAbcMeshSample* MergedMeshSample = new FAbcMeshSample();
-		for (const FCompressedAbcData& Data : CompressedMeshData)
-		{
-			AbcImporterUtilities::AppendMeshSample(MergedMeshSample, Data.AverageSample);
-		}
-		
-		// Forced to 1
-		LODModel.NumTexCoords = MergedMeshSample->NumUVSets;
-		SkeletalMesh->SetHasVertexColors(true);
-		SkeletalMesh->SetVertexColorGuid(FGuid::NewGuid());
-
 		/* Bounding box according to animation */
 		SkeletalMesh->SetImportedBounds(AbcFile->GetArchiveBounds().GetBox());
 
 		bool bBuildSuccess = false;
 		TArray<int32> MorphTargetVertexRemapping;
 		TArray<int32> UsedVertexIndicesForMorphs;
-		MergedMeshSample->TangentX.Empty();
-		MergedMeshSample->TangentY.Empty();
-		bBuildSuccess = BuildSkeletalMesh(LODModel, SkeletalMesh->GetRefSkeleton(), MergedMeshSample, MorphTargetVertexRemapping, UsedVertexIndicesForMorphs);
+
+		{
+			FAbcMeshSample MergedMeshSample; // Temporary mesh
+			for (const FCompressedAbcData& Data : CompressedMeshData)
+			{
+				AbcImporterUtilities::AppendMeshSample(&MergedMeshSample, Data.AverageSample);
+			}
+
+			MergedMeshSample.TangentX.Empty();
+			MergedMeshSample.TangentY.Empty();
+
+			// Forced to 1
+			LODModel.NumTexCoords = MergedMeshSample.NumUVSets;
+			SkeletalMesh->SetHasVertexColors(true);
+			SkeletalMesh->SetVertexColorGuid(FGuid::NewGuid());
+
+			bBuildSuccess = BuildSkeletalMesh(LODModel, SkeletalMesh->GetRefSkeleton(), &MergedMeshSample,
+				AbcFile->GetNumMaterialSlots(),
+				AbcFile->GetLookupMaterialSlot(),
+				MorphTargetVertexRemapping, UsedVertexIndicesForMorphs);
+		}
 
 		if (!bBuildSuccess)
 		{
@@ -756,18 +757,6 @@ TArray<UObject*> FAbcImporter::ImportAsSkeletalMesh(UObject* InParent, EObjectFl
 				WedgeOffset += CompressedData.AverageSample->Indices.Num();
 				VertexOffset += CompressedData.AverageSample->Vertices.Num();
 
-				const uint32 NumMaterials = CompressedData.MaterialNames.Num();
-				for (uint32 MaterialIndex = 0; MaterialIndex < NumMaterials; ++MaterialIndex)
-				{
-					const FString& MaterialName = CompressedData.MaterialNames[MaterialIndex];
-					UMaterialInterface* Material = AbcImporterUtilities::RetrieveMaterial(*AbcFile, MaterialName, InParent, Flags);
-					SkeletalMesh->GetMaterials().Add(FSkeletalMaterial(Material, true, false, FName(MaterialName), FName(MaterialName)));
-					if (Material != UMaterial::GetDefaultMaterial(MD_Surface))
-					{
-						Material->PostEditChange();
-					}
-				}
-
 				++ObjectIndex;
 			}
 
@@ -807,6 +796,27 @@ TArray<UObject*> FAbcImporter::ImportAsSkeletalMesh(UObject* InParent, EObjectFl
 			SkeletalMesh->CalculateInvRefMatrices();
 		}
 
+		UMaterial* DefaultMaterial = UMaterial::GetDefaultMaterial(MD_Surface);
+		check(DefaultMaterial);
+
+		// Build the material slots : one for each faceset
+		for (const FString& FaceSetName : AbcFile->GetUniqueFaceSetNames())
+		{
+			UMaterialInterface* Material = AbcImporterUtilities::RetrieveMaterial(*AbcFile, FaceSetName, InParent, Flags);
+			if (Material != DefaultMaterial)
+			{
+				Material->PostEditChange();
+			}
+
+			if (Material == nullptr)
+			{
+				Material = DefaultMaterial;
+			}
+
+			FName MaterialName(*FaceSetName);
+			SkeletalMesh->GetMaterials().Add(FSkeletalMaterial(Material, true, false, MaterialName, MaterialName));
+		}
+
 		SkeletalMesh->MarkPackageDirty();
 
 		Controller.NotifyPopulated();
@@ -1069,17 +1079,6 @@ const bool FAbcImporter::CompressAnimationDataUsingPCA(const FAbcCompressionSett
 			for (FAbcPolyMesh* PolyMesh : PolyMeshesToCompress)
 			{
 				AbcImporterUtilities::AppendMeshSample(&MergedZeroFrameSample, PolyMesh->GetTransformedFirstSample());
-
-				// QQ FUNCTIONALIZE
-				// Add material names from this mesh object
-				if (PolyMesh->FaceSetNames.Num() > 0)
-				{
-					CompressedData.MaterialNames.Append(PolyMesh->FaceSetNames);
-				}
-				else
-				{
-					CompressedData.MaterialNames.Add(NoFaceSetNameStr);
-				}
 			}
 
 			const uint32 NumVertices = AverageVertexData.Num();
@@ -1512,17 +1511,6 @@ const bool FAbcImporter::CompressAnimationDataUsingPCA(const FAbcCompressionSett
 				const float FrameStep = (MaxTimes[MeshIndex] - MinTimes[MeshIndex]) / (float)(NumSamples - 1);
 				AbcImporterUtilities::GenerateCompressedMeshData(CompressedData, NumUsedSingularValues, NumSamples, BasesMatrix, NormalsBasesMatrix, OutV, FrameStep, FMath::Max(MinTimes[MeshIndex], 0.0f));
 
-				// QQ FUNCTIONALIZE
-				// Add material names from this mesh object
-				if (PolyMeshesToCompress[MeshIndex]->FaceSetNames.Num() > 0)
-				{
-					CompressedData.MaterialNames.Append(PolyMeshesToCompress[MeshIndex]->FaceSetNames);
-				}
-				else
-				{
-					CompressedData.MaterialNames.Add(NoFaceSetNameStr);
-				}
-
 				if (bRunComparison)
 				{
 					CompareCompressionResult(Matrices[MeshIndex], NumSamples, NumUsedSingularValues, BasesMatrix, OutV, THRESH_POINTS_ARE_SAME);
@@ -1553,17 +1541,6 @@ const bool FAbcImporter::CompressAnimationDataUsingPCA(const FAbcCompressionSett
 		{
 			CompressedData.AverageSample = new FAbcMeshSample(*ConstantPolyMesh->GetFirstSample());
 		}
-
-		// QQ FUNCTIONALIZE
-		// Add material names from this mesh object
-		if (ConstantPolyMesh->FaceSetNames.Num() > 0)
-		{
-			CompressedData.MaterialNames.Append(ConstantPolyMesh->FaceSetNames);
-		}
-		else
-		{
-			CompressedData.MaterialNames.Add(NoFaceSetNameStr);
-		}
 	}
 		
 	return bResult;
@@ -1699,7 +1676,7 @@ const uint32 FAbcImporter::GetNumMeshTracks() const
 	return (AbcFile != nullptr) ? AbcFile->GetNumPolyMeshes() : 0;
 }
 
-void FAbcImporter::GenerateMeshDescriptionFromSample(const FAbcMeshSample* Sample, FMeshDescription* MeshDescription, UStaticMesh* StaticMesh)
+void FAbcImporter::GenerateMeshDescriptionFromSample(const TArray<FString>& UniqueFaceSetNames, const TArray<int32>& LookupMaterialSlot, const FAbcMeshSample* Sample, FMeshDescription* MeshDescription)
 {
 	if (MeshDescription == nullptr)
 	{
@@ -1719,11 +1696,11 @@ void FAbcImporter::GenerateMeshDescriptionFromSample(const FAbcMeshSample* Sampl
 
 	//Speedtree use UVs to store is data
 	VertexInstanceUVs.SetNumChannels(Sample->NumUVSets);
-	
-	for (int32 MatIndex = 0; MatIndex < StaticMesh->GetStaticMaterials().Num(); ++MatIndex)
+
+	for (const FString& FaceSetName: UniqueFaceSetNames)
 	{
 		const FPolygonGroupID PolygonGroupID = MeshDescription->CreatePolygonGroup();
-		PolygonGroupImportedMaterialSlotNames[PolygonGroupID] = StaticMesh->GetStaticMaterials()[MatIndex].ImportedMaterialSlotName;
+		PolygonGroupImportedMaterialSlotNames[PolygonGroupID] = FName(*FaceSetName);
 	}
 
 	// position
@@ -1786,7 +1763,9 @@ void FAbcImporter::GenerateMeshDescriptionFromSample(const FAbcMeshSample* Sampl
 			CornerVertexIDs[Corner] = VertexID;
 		}
 
-		const FPolygonGroupID PolygonGroupID(Sample->MaterialIndices[TriangleIndex]);
+		int32 MaterialSlotID = LookupMaterialSlot[Sample->MaterialIndices[TriangleIndex]];
+		const FPolygonGroupID PolygonGroupID(MaterialSlotID);
+
 		// Insert a polygon into the mesh
 		MeshDescription->CreatePolygon(PolygonGroupID, CornerVertexInstanceIDs);
 	}
@@ -1794,7 +1773,9 @@ void FAbcImporter::GenerateMeshDescriptionFromSample(const FAbcMeshSample* Sampl
 	FStaticMeshOperations::ConvertSmoothGroupToHardEdges(Sample->SmoothingGroupIndices, *MeshDescription);
 }
 
-bool FAbcImporter::BuildSkeletalMesh( FSkeletalMeshLODModel& LODModel, const FReferenceSkeleton& RefSkeleton, FAbcMeshSample* Sample, TArray<int32>& OutMorphTargetVertexRemapping, TArray<int32>& OutUsedVertexIndicesForMorphs)
+bool FAbcImporter::BuildSkeletalMesh( FSkeletalMeshLODModel& LODModel, const FReferenceSkeleton& RefSkeleton, FAbcMeshSample* Sample, 
+	int32 NumMaterialSlots, const TArray<int32> LookupMaterialSlot, 
+	TArray<int32>& OutMorphTargetVertexRemapping, TArray<int32>& OutUsedVertexIndicesForMorphs)
 {
 	// Module manager is not thread safe, so need to prefetch before parallelfor
 	IMeshUtilities& MeshUtilities = FModuleManager::Get().LoadModuleChecked<IMeshUtilities>("MeshUtilities");
@@ -1815,13 +1796,13 @@ bool FAbcImporter::BuildSkeletalMesh( FSkeletalMeshLODModel& LODModel, const FRe
 	Faces.AddZeroed(NumFaces);
 
 	TArray<FMeshSection> MeshSections;
-	MeshSections.AddDefaulted(Sample->NumMaterials);
+	MeshSections.AddDefaulted(NumMaterialSlots);
 
 	// Process all the faces and add to their respective mesh section
 	for (uint32 FaceIndex = 0; FaceIndex < NumFaces; ++FaceIndex)
 	{
 		const uint32 FaceOffset = FaceIndex * 3;
-		const int32 MaterialIndex = Sample->MaterialIndices[FaceIndex];
+		const int32 MaterialIndex = LookupMaterialSlot[Sample->MaterialIndices[FaceIndex]];
 
 		check(MeshSections.IsValidIndex(MaterialIndex));
 
diff --git a/Engine/Plugins/Importers/AlembicImporter/Source/AlembicLibrary/Public/AbcImporter.h b/Engine/Plugins/Importers/AlembicImporter/Source/AlembicLibrary/Public/AbcImporter.h
--- a/Engine/Plugins/Importers/AlembicImporter/Source/AlembicLibrary/Public/AbcImporter.h
+++ b/Engine/Plugins/Importers/AlembicImporter/Source/AlembicLibrary/Public/AbcImporter.h
@@ -62,8 +62,6 @@ struct FCompressedAbcData
 	TArray<TArray<float>> CurveValues;
 	/** Contains the time key values for each individual base */
 	TArray<TArray<float>> TimeValues;
-	/** Material names used for retrieving created materials */
-	TArray<FString> MaterialNames;
 };
 
 /** Mesh section used for chunking the mesh data during Skeletal mesh building */
@@ -179,14 +177,15 @@ private:
 	* @param Name - Name for the static mesh
 	* @param Flags - Object flags
 	* @param NumMaterials - Number of materials to add
-	* @param FaceSetNames - Face set names used for retrieving the materials
-	* @param MeshDescription - The MeshDescription from which the static mesh should be constructed
+	* @param UniqueFaceSetNames - The array of unique face set (material slot) names for merged mesh
+	* @param LookupMaterialSlot - Mapping from faceset index (in flattened list) to material slot
+	* @param Sample - The FAbcMeshSample from which the static mesh should be constructed
 	* @return UStaticMesh*
 	*/
-	UStaticMesh* CreateStaticMeshFromSample(UObject* InParent, const FString& Name, EObjectFlags Flags, const uint32 NumMaterials, const TArray<FString>& FaceSetNames, const FAbcMeshSample* Sample);
+	UStaticMesh* CreateStaticMeshFromSample(UObject* InParent, const FString& Name, EObjectFlags Flags, const TArray<FString>& UniqueFaceSetNames, const TArray<int32>& LookupMaterialSlot, const FAbcMeshSample* Sample);
 
 	/** Generates and populate a FMeshDescription instance from the given sample*/
-	void GenerateMeshDescriptionFromSample(const FAbcMeshSample* Sample, FMeshDescription* MeshDescription, UStaticMesh* StaticMesh);
+	void GenerateMeshDescriptionFromSample(const TArray<FString>& UniqueFaceSetNames, const TArray<int32>& LookupMaterialSlot, const FAbcMeshSample* Sample, FMeshDescription* MeshDescription);
 	
 	/** Compresses the imported animation data, returns true if compression was successful and compressed data was populated */
 	const bool CompressAnimationDataUsingPCA(const FAbcCompressionSettings& InCompressionSettings, const bool bRunComparison = false);	
@@ -198,7 +197,9 @@ private:
 	void CompareCompressionResult(const TArray64<float>& OriginalMatrix, const uint32 NumSamples, const uint32 NumUsedSingularValues, const TArrayView64<float>& OutU, const TArray64<float>& OutV, const float Tolerance);
 	
 	/** Build a skeletal mesh from the PCA compressed data */
-	bool BuildSkeletalMesh(FSkeletalMeshLODModel& LODModel, const FReferenceSkeleton& RefSkeleton, FAbcMeshSample* Sample, TArray<int32>& OutMorphTargetVertexRemapping, TArray<int32>& OutUsedVertexIndicesForMorphs);
+	bool BuildSkeletalMesh(FSkeletalMeshLODModel& LODModel, const FReferenceSkeleton& RefSkeleton, FAbcMeshSample* Sample, 
+		int32 NumMaterialSlots, const TArray<int32> LookupMaterialSlot, 
+		TArray<int32>& OutMorphTargetVertexRemapping, TArray<int32>& OutUsedVertexIndicesForMorphs);
 	
 	/** Generate morph target vertices from the PCA compressed bases */
 	void GenerateMorphTargetVertices(FAbcMeshSample* BaseSample, TArray<FMorphTargetDelta> &MorphDeltas, FAbcMeshSample* AverageSample, uint32 WedgeOffset, const TArray<int32>& RemapIndices, const TArray<int32>& UsedVertexIndicesForMorphs, const uint32 VertexOffset, const uint32 IndexOffset);
