[CrashReport]UnrealEditor-ControlRig!TSet<TTuple<FRigElementKey,int>,TDefaultMapHashableKeyFuncs<FRigElementKey,int,0>,FDefaultSetAllocator>::FindIndexByHash<FRigElementKey>(unsigned int,FRigElementKey const &) [Set.h:900]
#fix ensure that animation recording and Alembic imports are not transacted (undo/redo-able)

diff --git a/Engine/Plugins/Importers/AlembicImporter/Source/AlembicLibrary/Private/AbcImporter.cpp b/Engine/Plugins/Importers/AlembicImporter/Source/AlembicLibrary/Private/AbcImporter.cpp
--- a/Engine/Plugins/Importers/AlembicImporter/Source/AlembicLibrary/Private/AbcImporter.cpp
+++ b/Engine/Plugins/Importers/AlembicImporter/Source/AlembicLibrary/Private/AbcImporter.cpp
@@ -139,11 +139,12 @@ const EAbcImportError FAbcImporter::ImportTrackData(const int32 InNumThreads, UA
 }
 
 template<typename T>
-T* FAbcImporter::CreateObjectInstance(UObject*& InParent, const FString& ObjectName, const EObjectFlags Flags)
+T* FAbcImporter::CreateObjectInstance(UObject*& InParent, const FString& ObjectName, const EObjectFlags Flags, bool& bObjectAlreadyExists)
 {
 	// Parent package to place new asset
 	UPackage* Package = nullptr;
 	FString NewPackageName;
+	bObjectAlreadyExists = false;
 
 	// Setup package name and create one accordingly
 	NewPackageName = FPackageName::GetLongPackagePath(InParent->GetOutermost()->GetPathName()) + TEXT("/") + ObjectName;
@@ -158,6 +159,8 @@ T* FAbcImporter::CreateObjectInstance(UObject*& InParent, const FString& ObjectN
 	if (ExistingTypedObject != nullptr)
 	{
 		ExistingTypedObject->PreEditChange(nullptr);
+		bObjectAlreadyExists = true;
+		return ExistingTypedObject;
 	}
 	else if (ExistingObject != nullptr)
 	{
@@ -186,7 +189,8 @@ T* FAbcImporter::CreateObjectInstance(UObject*& InParent, const FString& ObjectN
 
 UStaticMesh* FAbcImporter::CreateStaticMeshFromSample(UObject* InParent, const FString& Name, EObjectFlags Flags, const TArray<FString>& UniqueFaceSetNames, const TArray<int32>& LookupMaterialSlot, const FAbcMeshSample* Sample)
 {
-	UStaticMesh* StaticMesh = CreateObjectInstance<UStaticMesh>(InParent, Name, Flags);
+	bool bObjectAlreadyExists = false;
+	UStaticMesh* StaticMesh = CreateObjectInstance<UStaticMesh>(InParent, Name, Flags, bObjectAlreadyExists);
 
 	// Only import data if a valid object was created
 	if (StaticMesh)
@@ -345,8 +349,9 @@ const TArray<UStaticMesh*> FAbcImporter::ImportAsStaticMesh(UObject* InParent, E
 
 UGeometryCache* FAbcImporter::ImportAsGeometryCache(UObject* InParent, EObjectFlags Flags)
 {
-	// Create a GeometryCache instance 
-	UGeometryCache* GeometryCache = CreateObjectInstance<UGeometryCache>(InParent, InParent != GetTransientPackage() ? FPaths::GetBaseFilename(InParent->GetName()) : (FPaths::GetBaseFilename(AbcFile->GetFilePath()) + "_" + FGuid::NewGuid().ToString()), Flags);
+	// Create a GeometryCache instance
+	bool bObjectAlreadyExists = false;
+	UGeometryCache* GeometryCache = CreateObjectInstance<UGeometryCache>(InParent, InParent != GetTransientPackage() ? FPaths::GetBaseFilename(InParent->GetName()) : (FPaths::GetBaseFilename(AbcFile->GetFilePath()) + "_" + FGuid::NewGuid().ToString()), Flags, bObjectAlreadyExists);
 
 	// Only import data if a valid object was created
 	if (GeometryCache)
@@ -602,7 +607,8 @@ TArray<UObject*> FAbcImporter::ImportAsSkeletalMesh(UObject* InParent, EObjectFl
 	USkeletalMesh* ExistingSkeletalMesh = FindObject<USkeletalMesh>(InParent, *SanitizedObjectName);	
 	FSkinnedMeshComponentRecreateRenderStateContext* RecreateExistingRenderStateContext = ExistingSkeletalMesh ? new FSkinnedMeshComponentRecreateRenderStateContext(ExistingSkeletalMesh, false) : nullptr;
 	
-	USkeletalMesh* SkeletalMesh = CreateObjectInstance<USkeletalMesh>(InParent, ObjectName, Flags);
+	bool bMeshAlreadyExists = false;
+	USkeletalMesh* SkeletalMesh = CreateObjectInstance<USkeletalMesh>(InParent, ObjectName, Flags, bMeshAlreadyExists);
 
 	// Only import data if a valid object was created
 	if (SkeletalMesh)
@@ -668,7 +674,8 @@ TArray<UObject*> FAbcImporter::ImportAsSkeletalMesh(UObject* InParent, EObjectFl
 
 		// Create the skeleton object
 		FString SkeletonName = FString::Printf(TEXT("%s_Skeleton"), *SkeletalMesh->GetName());
-		USkeleton* Skeleton = CreateObjectInstance<USkeleton>(InParent, SkeletonName, Flags);
+		bool bSkeletonAlreadyExists = false;
+		USkeleton* Skeleton = CreateObjectInstance<USkeleton>(InParent, SkeletonName, Flags, bSkeletonAlreadyExists);
 
 		// Merge bones to the selected skeleton
 		check(Skeleton->MergeAllBonesToBoneTree(SkeletalMesh));
@@ -680,7 +687,8 @@ TArray<UObject*> FAbcImporter::ImportAsSkeletalMesh(UObject* InParent, EObjectFl
 		}
 
 		// Create animation sequence for the skeleton
-		UAnimSequence* Sequence = CreateObjectInstance<UAnimSequence>(InParent, FString::Printf(TEXT("%s_Animation"), *SkeletalMesh->GetName()), Flags);
+		bool bSequenceAlreadyExists = false;
+		UAnimSequence* Sequence = CreateObjectInstance<UAnimSequence>(InParent, FString::Printf(TEXT("%s_Animation"), *SkeletalMesh->GetName()), Flags, bSequenceAlreadyExists);
 		Sequence->SetSkeleton(Skeleton);
 
 		int32 ObjectIndex = 0;
@@ -690,13 +698,14 @@ TArray<UObject*> FAbcImporter::ImportAsSkeletalMesh(UObject* InParent, EObjectFl
 
 		IAnimationDataController& Controller = Sequence->GetController();
 
-		Controller.OpenBracket(LOCTEXT("ImportAsSkeletalMesh", "Importing Alembic Animation"));
+		const bool bShouldTransact = bSequenceAlreadyExists;
+		Controller.OpenBracket(LOCTEXT("ImportAsSkeletalMesh", "Importing Alembic Animation"), bShouldTransact);
 		Controller.InitializeModel();
 
 		const FFrameRate FrameRate(FMath::RoundToInt(AbcFile->GetFramerate()), 1);
-		Controller.SetFrameRate(FrameRate);	
+		Controller.SetFrameRate(FrameRate, bShouldTransact);	
 		const FFrameNumber FrameNumber = FrameRate.AsFrameNumber(AbcFile->GetImportLength());
-		Controller.SetNumberOfFrames(FrameNumber);
+		Controller.SetNumberOfFrames(FrameNumber, bShouldTransact);
 
 		Sequence->ImportFileFramerate = static_cast<float>(FrameRate.AsDecimal());
 		Sequence->ImportResampleFramerate = static_cast<int32>(FrameRate.AsInterval());
@@ -745,7 +754,7 @@ TArray<UObject*> FAbcImporter::ImportAsSkeletalMesh(UObject* InParent, EObjectFl
 							FName ConstCurveName = *CurveName;
 
 							// Sets up the morph target curves with the sample values and time keys
-							SetupMorphTargetCurves(Skeleton, ConstCurveName, Sequence, CurveValues, TimeValues, Controller);
+							SetupMorphTargetCurves(Skeleton, ConstCurveName, Sequence, CurveValues, TimeValues, Controller, bShouldTransact);
 						}
 						else
 						{
@@ -783,8 +792,8 @@ TArray<UObject*> FAbcImporter::ImportAsSkeletalMesh(UObject* InParent, EObjectFl
 				const FReferenceSkeleton& RefSkeleton = SkeletalMesh->GetRefSkeleton();
 				const TArray<FMeshBoneInfo>& BonesInfo = RefSkeleton.GetRawRefBoneInfo();
 				
-				Controller.AddBoneCurve(BonesInfo[0].Name);
-				Controller.SetBoneTrackKeys(BonesInfo[0].Name, RootBoneTrack.PosKeys, RootBoneTrack.RotKeys, RootBoneTrack.ScaleKeys);
+				Controller.AddBoneCurve(BonesInfo[0].Name, bShouldTransact);
+				Controller.SetBoneTrackKeys(BonesInfo[0].Name, RootBoneTrack.PosKeys, RootBoneTrack.RotKeys, RootBoneTrack.ScaleKeys, bShouldTransact);
 			}
 
 			// Set recompute tangent flag on skeletal mesh sections
@@ -821,7 +830,7 @@ TArray<UObject*> FAbcImporter::ImportAsSkeletalMesh(UObject* InParent, EObjectFl
 
 		Controller.NotifyPopulated();
 
-		Controller.CloseBracket();
+		Controller.CloseBracket(bShouldTransact);
 
 		Sequence->PostEditChange();
 		Sequence->SetPreviewMesh(SkeletalMesh);
@@ -850,7 +859,7 @@ TArray<UObject*> FAbcImporter::ImportAsSkeletalMesh(UObject* InParent, EObjectFl
 	return GeneratedObjects;
 }
 
-void FAbcImporter::SetupMorphTargetCurves(USkeleton* Skeleton, FName ConstCurveName, UAnimSequence* Sequence, const TArray<float> &CurveValues, const TArray<float>& TimeValues, IAnimationDataController& Controller)
+void FAbcImporter::SetupMorphTargetCurves(USkeleton* Skeleton, FName ConstCurveName, UAnimSequence* Sequence, const TArray<float> &CurveValues, const TArray<float>& TimeValues, IAnimationDataController& Controller, bool bShouldTransact)
 {
 	// Need curve metadata for the AnimSequence to playback. Can be either on the Skeleton or SKelMesh,
 	// but by default for FBX import it's on the Skeleton so do the same for Alembic
@@ -859,7 +868,7 @@ void FAbcImporter::SetupMorphTargetCurves(USkeleton* Skeleton, FName ConstCurveN
 	Skeleton->AccumulateCurveMetaData(ConstCurveName, bMaterialCurve, bMorphTargetCurve);
 
 	FAnimationCurveIdentifier CurveId(ConstCurveName, ERawCurveTrackTypes::RCT_Float);
-	Controller.AddCurve(CurveId);
+	Controller.AddCurve(CurveId, EAnimAssetCurveFlags::AACF_Editable, bShouldTransact);
 
 	const FFloatCurve* NewCurve = Sequence->GetDataModel()->FindFloatCurve(CurveId);
 	ensure(NewCurve);
@@ -881,7 +890,7 @@ void FAbcImporter::SetupMorphTargetCurves(USkeleton* Skeleton, FName ConstCurveN
 		RichCurve.SetKeyTangentWeightMode(NewKeyHandle, NewTangentWeightMode);
 	}
 
-	Controller.SetCurveKeys(CurveId, RichCurve.GetConstRefOfKeys());
+	Controller.SetCurveKeys(CurveId, RichCurve.GetConstRefOfKeys(), bShouldTransact);
 }
 
 void FAbcImporter::SetMetaData(const TArray<UObject*>& Objects)
diff --git a/Engine/Plugins/Importers/AlembicImporter/Source/AlembicLibrary/Public/AbcImporter.h b/Engine/Plugins/Importers/AlembicImporter/Source/AlembicLibrary/Public/AbcImporter.h
--- a/Engine/Plugins/Importers/AlembicImporter/Source/AlembicLibrary/Public/AbcImporter.h
+++ b/Engine/Plugins/Importers/AlembicImporter/Source/AlembicLibrary/Public/AbcImporter.h
@@ -168,7 +168,7 @@ private:
 	* @param ObjectName - Name to be used for the created object
 	* @param Flags - Object creation flags
 	*/
-	template<typename T> T* CreateObjectInstance(UObject*& InParent, const FString& ObjectName, const EObjectFlags Flags);
+	template<typename T> T* CreateObjectInstance(UObject*& InParent, const FString& ObjectName, const EObjectFlags Flags, bool& bObjectAlreadyExists);
 	
 	/**
 	* Creates a Static mesh from the given mesh description
@@ -205,7 +205,7 @@ private:
 	void GenerateMorphTargetVertices(FAbcMeshSample* BaseSample, TArray<FMorphTargetDelta> &MorphDeltas, FAbcMeshSample* AverageSample, uint32 WedgeOffset, const TArray<int32>& RemapIndices, const TArray<int32>& UsedVertexIndicesForMorphs, const uint32 VertexOffset, const uint32 IndexOffset);
 	
 	/** Set up correct morph target weights from the PCA compressed data */
-	void SetupMorphTargetCurves(USkeleton* Skeleton, FName ConstCurveName, UAnimSequence* Sequence, const TArray<float> &CurveValues, const TArray<float>& TimeValues, IAnimationDataController& Controller);
+	void SetupMorphTargetCurves(USkeleton* Skeleton, FName ConstCurveName, UAnimSequence* Sequence, const TArray<float> &CurveValues, const TArray<float>& TimeValues, IAnimationDataController& Controller, bool bShouldTransact);
 	
 	/** Set the Alembic archive metadata on the given objects */
 	void SetMetaData(const TArray<UObject*>& Objects);
diff --git a/Engine/Source/Editor/SequenceRecorder/Private/AnimationRecorder.cpp b/Engine/Source/Editor/SequenceRecorder/Private/AnimationRecorder.cpp
--- a/Engine/Source/Editor/SequenceRecorder/Private/AnimationRecorder.cpp
+++ b/Engine/Source/Editor/SequenceRecorder/Private/AnimationRecorder.cpp
@@ -169,7 +169,7 @@ bool FAnimationRecorder::TriggerRecordAnimation(USkeletalMeshComponent* Componen
 
 		SetSampleRateAndLength(SampleRate, MaximumLength);
 		
-		Parent = CreatePackage( *ValidatedAssetPath);
+		Parent = CreatePackage( *ValidatedAssetPath);	
 	}
 
 	UObject* const Object = LoadObject<UObject>(Parent, *ValidatedAssetName, nullptr, LOAD_Quiet, nullptr);
@@ -180,7 +180,10 @@ bool FAnimationRecorder::TriggerRecordAnimation(USkeletalMeshComponent* Componen
 		return false;		// failed
 	}
 
-	// If not, create new one now.
+	// If not, create new one now. (also means we should not transact the recording)
+	const bool bExistingAnimSequence = FindObject<UAnimSequence>(Parent, *ValidatedAssetName) != nullptr;
+	bTransactRecording &= bExistingAnimSequence;
+
 	UAnimSequence* const NewSeq = NewObject<UAnimSequence>(Parent, *ValidatedAssetName, RF_Public | RF_Standalone);
 	if (NewSeq)
 	{
