Fix alembic importer when checks are disabled

diff --git a/Engine/Plugins/Importers/AlembicImporter/Source/AlembicLibrary/Private/AbcImporter.cpp b/Engine/Plugins/Importers/AlembicImporter/Source/AlembicLibrary/Private/AbcImporter.cpp
--- a/Engine/Plugins/Importers/AlembicImporter/Source/AlembicLibrary/Private/AbcImporter.cpp
+++ b/Engine/Plugins/Importers/AlembicImporter/Source/AlembicLibrary/Private/AbcImporter.cpp
@@ -677,7 +677,7 @@ TArray<UObject*> FAbcImporter::ImportAsSkeletalMesh(UObject* InParent, EObjectFl
 		USkeleton* Skeleton = CreateObjectInstance<USkeleton>(InParent, SkeletonName, Flags, bSkeletonAlreadyExists);
 
 		// Merge bones to the selected skeleton
-		check(Skeleton->MergeAllBonesToBoneTree(SkeletalMesh));
+		ensure(Skeleton->MergeAllBonesToBoneTree(SkeletalMesh));
 		Skeleton->MarkPackageDirty();
 		if (SkeletalMesh->GetSkeleton() != Skeleton)
 		{
