Alembic: Fixed crash when reimporting skeletal mesh previously imported from Alembic.

diff --git a/Engine/Plugins/Importers/AlembicImporter/Source/AlembicLibrary/Private/AbcImporter.cpp b/Engine/Plugins/Importers/AlembicImporter/Source/AlembicLibrary/Private/AbcImporter.cpp
--- a/Engine/Plugins/Importers/AlembicImporter/Source/AlembicLibrary/Private/AbcImporter.cpp
+++ b/Engine/Plugins/Importers/AlembicImporter/Source/AlembicLibrary/Private/AbcImporter.cpp
@@ -617,7 +617,6 @@ TArray<UObject*> FAbcImporter::ImportAsSkeletalMesh(UObject* InParent, EObjectFl
 
 		// Retrieve the imported resource structure and allocate a new LOD model
 		FSkeletalMeshModel* ImportedModel = SkeletalMesh->GetImportedModel();
-		check(ImportedModel->LODModels.Num() == 0);
 		ImportedModel->LODModels.Empty();
 		ImportedModel->EmptyOriginalReductionSourceMeshData();
 		ImportedModel->LODModels.Add(new FSkeletalMeshLODModel());
@@ -633,7 +632,10 @@ TArray<UObject*> FAbcImporter::ImportAsSkeletalMesh(UObject* InParent, EObjectFl
 		const FTransform BoneTransform;
 		{
 			FReferenceSkeletonModifier RefSkelModifier(SkeletalMesh->GetRefSkeleton(), SkeletalMesh->GetSkeleton());
-			RefSkelModifier.Add(BoneInfo, BoneTransform);
+			if (RefSkelModifier.FindBoneIndex(BoneInfo.Name) == INDEX_NONE)
+			{
+				RefSkelModifier.Add(BoneInfo, BoneTransform);
+			}
 		}
 
 
