Alembic: Fixed potential vulnerability issue when importing an Alembic with indices that are out of bounds.

diff --git a/Engine/Plugins/Importers/AlembicImporter/Source/AlembicLibrary/Private/AbcImportUtilities.cpp b/Engine/Plugins/Importers/AlembicImporter/Source/AlembicLibrary/Private/AbcImportUtilities.cpp
--- a/Engine/Plugins/Importers/AlembicImporter/Source/AlembicLibrary/Private/AbcImportUtilities.cpp
+++ b/Engine/Plugins/Importers/AlembicImporter/Source/AlembicLibrary/Private/AbcImportUtilities.cpp
@@ -480,6 +480,20 @@ bool AbcImporterUtilities::GenerateAbcMeshSampleDataForFrame(const Alembic::AbcG
 	{
 		Alembic::Abc::Int32ArraySamplePtr IndicesSample = MeshSample.getFaceIndices();
 		bRetrievalResult &= RetrieveTypedAbcData<Alembic::Abc::Int32ArraySamplePtr, uint32>(IndicesSample, Sample->Indices);
+		if (bRetrievalResult)
+		{
+			// Validate that the indices coming from the Alembic are valid, in case it is corrupted
+			const uint32 MaxIndex = IntCastChecked<uint32>(Sample->Vertices.Num());
+			for (uint32& Index : Sample->Indices)
+			{
+				if (Index >= MaxIndex)
+				{
+					TSharedRef<FTokenizedMessage> Message = FTokenizedMessage::Create(EMessageSeverity::Error, LOCTEXT("FoundInvalidIndex", "Unable to import mesh due to index out of bound found."));
+					FAbcImportLogger::AddImportMessage(Message);
+					return false;
+				}
+			}
+		}
 		if (bNeedsTriangulation)
 		{
 			TriangulateIndexBuffer(FaceCounts, Sample->Indices);
diff --git a/Engine/Plugins/Importers/AlembicImporter/Source/AlembicLibrary/Private/AbcPolyMesh.cpp b/Engine/Plugins/Importers/AlembicImporter/Source/AlembicLibrary/Private/AbcPolyMesh.cpp
--- a/Engine/Plugins/Importers/AlembicImporter/Source/AlembicLibrary/Private/AbcPolyMesh.cpp
+++ b/Engine/Plugins/Importers/AlembicImporter/Source/AlembicLibrary/Private/AbcPolyMesh.cpp
@@ -173,6 +173,10 @@ void FAbcPolyMesh::SetFrameAndTime(const float InTime, const int32 FrameIndex, c
 			const ESampleReadFlags SampleReadFlagsOverride = EnumHasAnyFlags(InFlags, EFrameReadFlags::PositionAndNormalOnly) ? ESampleReadFlags::Positions |  ESampleReadFlags::Normals : SampleReadFlags;
 			WriteSample->Copy(FirstSample, SampleReadFlagsOverride);
 			const bool bValidSample = AbcImporterUtilities::GenerateAbcMeshSampleDataForFrame(Schema, SampleSelector, WriteSample, SampleReadFlagsOverride, InTime == MinTime);
+			if (!bValidSample)
+			{
+				return;
+			}
 			// Check whether or not the number of normal indices matches with the first frame
 			const bool bMatchingIndices = FirstSample != nullptr && FirstSample->Indices.Num() == WriteSample->Indices.Num();
 			// Make sure in case of recomputing normals we enforece using the first sample data (otherwise we'll be using loaded or incorrectly calculated normals)
