Alembic: Fixed crash when importing an Alembic that reports having UV but not returning them when queried.

diff --git a/Engine/Plugins/Importers/AlembicImporter/Source/AlembicLibrary/Private/AbcImportUtilities.cpp b/Engine/Plugins/Importers/AlembicImporter/Source/AlembicLibrary/Private/AbcImportUtilities.cpp
--- a/Engine/Plugins/Importers/AlembicImporter/Source/AlembicLibrary/Private/AbcImportUtilities.cpp
+++ b/Engine/Plugins/Importers/AlembicImporter/Source/AlembicLibrary/Private/AbcImportUtilities.cpp
@@ -717,15 +717,22 @@ void AbcImporterUtilities::ReadUVSetData(Alembic::AbcGeom::IV2fGeomParam &UVCoor
 	{
 		Alembic::Abc::UInt32ArraySamplePtr UVIndiceSample = UVCoordinateParameter.getIndexProperty().getValue(FrameSelector);
 		TArray<uint32> UVIndices;
-		RetrieveTypedAbcData<Alembic::Abc::UInt32ArraySamplePtr, uint32>(UVIndiceSample, UVIndices);
+		const bool bSuccess = RetrieveTypedAbcData<Alembic::Abc::UInt32ArraySamplePtr, uint32>(UVIndiceSample, UVIndices);
+		if (bSuccess)
+		{
+			if (bNeedsTriangulation)
+			{
+				TriangulateIndexBuffer(FaceCounts, UVIndices);
+			}
 
-		if (bNeedsTriangulation)
+			// Expand UV array
+			ExpandVertexAttributeArray<FVector2f>(UVIndices, OutUVs);
+		}
+		else
 		{
-			TriangulateIndexBuffer(FaceCounts, UVIndices);
+			// Fallback to ensure there's always UV data present since it's needed later to compute tangents
+			OutUVs.AddZeroed(MeshIndices.Num());
 		}
-
-		// Expand UV array
-		ExpandVertexAttributeArray<FVector2f>(UVIndices, OutUVs);
 	}
 	else if (OutUVs.Num())
 	{
