Fixed potential crash when importing geometry cache with varying topologies from Alembic.

diff --git a/Engine/Plugins/Importers/AlembicImporter/Source/AlembicLibrary/Private/AbcImportUtilities.cpp b/Engine/Plugins/Importers/AlembicImporter/Source/AlembicLibrary/Private/AbcImportUtilities.cpp
--- a/Engine/Plugins/Importers/AlembicImporter/Source/AlembicLibrary/Private/AbcImportUtilities.cpp
+++ b/Engine/Plugins/Importers/AlembicImporter/Source/AlembicLibrary/Private/AbcImportUtilities.cpp
@@ -1071,7 +1071,7 @@ FAbcMeshSample* AbcImporterUtilities::MergeMeshSamples(const TArray<const FAbcMe
 		const uint32 MaterialIndicesOffset = MergedSample->MaterialIndices.Num();
 		const uint32 SmoothingGroupIndicesOffset = MergedSample->SmoothingGroupIndices.Num();
 
-		ensureMsgf(MaterialIndicesOffset == SmoothingGroupIndicesOffset, TEXT("Material and smoothing group indice count should match"));
+		ensureMsgf(MaterialIndicesOffset == SmoothingGroupIndicesOffset, TEXT("Material and smoothing group indices count should match (%d != %d))"), MaterialIndicesOffset, SmoothingGroupIndicesOffset);
 
 		// Per Face material and smoothing group index
 		MergedSample->MaterialIndices.Append(Sample->MaterialIndices);
@@ -1152,7 +1152,7 @@ void AbcImporterUtilities::AppendMeshSample(FAbcMeshSample* MeshSampleOne, const
 	const uint32 MaterialIndicesOffset = MeshSampleOne->MaterialIndices.Num();
 	const uint32 SmoothingGroupIndicesOffset = MeshSampleOne->SmoothingGroupIndices.Num();
 
-	ensureMsgf(MaterialIndicesOffset == SmoothingGroupIndicesOffset, TEXT("Material and smoothing group indice count should match"));
+	ensureMsgf(MaterialIndicesOffset == SmoothingGroupIndicesOffset, TEXT("Material and smoothing group indices count should match (%d != %d))"), MaterialIndicesOffset, SmoothingGroupIndicesOffset);
 
 	// Per Face material and smoothing group index
 	MeshSampleOne->MaterialIndices.Append(MeshSampleTwo->MaterialIndices);
@@ -1704,6 +1704,16 @@ void AbcImporterUtilities::MergePolyMeshesToMeshData(int32 FrameIndex, int32 Fra
 	}
 	PreviousNumVertices = MergedSample.Vertices.Num();
 
+	// For some reason, merged meshes (especially with varying topologies) can end up with an invalid number of material indices.
+	// Here, the material indices is filled up with 0 (the first material) to match the number of expected triangles so that the
+	// generated geometry cache mesh data from GeometryCacheDataForMeshSample is complete.
+	const int32 NumTriangles = MergedSample.Indices.Num() / 3;
+	const int32 NumMatIndices = MergedSample.MaterialIndices.Num();
+	if (!ensureMsgf(NumMatIndices >= NumTriangles, TEXT("Patching merged meshes (%d) since there are fewer material indices than triangles (%d < %d)."), PolyMeshes.Num(), NumMatIndices, NumTriangles))
+	{
+		const int32 Delta = NumTriangles - NumMatIndices;
+		MergedSample.MaterialIndices.AddZeroed(Delta);
+	}
 
 	// Generate the mesh data for this sample
 	AbcImporterUtilities::GeometryCacheDataForMeshSample(MeshData, &MergedSample, 0, SecondsPerFrame, bUseVelocitiesAsMotionVectors, bStoreImportedVertexNumbers);
