Fixed potential crash during Alembic import when the mesh is missing necessary attributes like positions, face counts and indices.

#rnx

diff --git a/Engine/Plugins/Importers/AlembicImporter/Source/AlembicLibrary/Private/AbcImportUtilities.cpp b/Engine/Plugins/Importers/AlembicImporter/Source/AlembicLibrary/Private/AbcImportUtilities.cpp
--- a/Engine/Plugins/Importers/AlembicImporter/Source/AlembicLibrary/Private/AbcImportUtilities.cpp
+++ b/Engine/Plugins/Importers/AlembicImporter/Source/AlembicLibrary/Private/AbcImportUtilities.cpp
@@ -443,6 +443,12 @@ bool AbcImporterUtilities::GenerateAbcMeshSampleDataForFrame(const Alembic::AbcG
 	{
 		Alembic::Abc::P3fArraySamplePtr PositionsSample = MeshSample.getPositions();
 		bRetrievalResult &= RetrieveTypedAbcData<Alembic::Abc::P3fArraySamplePtr, FVector3f>(PositionsSample, Sample->Vertices);
+		if (!bRetrievalResult)
+		{
+			TSharedRef<FTokenizedMessage> Message = FTokenizedMessage::Create(EMessageSeverity::Error, LOCTEXT("NoVertices", "Could not retrieve mesh vertices."));
+			FAbcImportLogger::AddImportMessage(Message);
+			return false;
+		}
 	}
 	
 	if (EnumHasAnyFlags(ReadFlags, ESampleReadFlags::Velocities))
@@ -461,6 +467,12 @@ bool AbcImporterUtilities::GenerateAbcMeshSampleDataForFrame(const Alembic::AbcG
 	{
 		Alembic::Abc::Int32ArraySamplePtr FaceCountsSample = MeshSample.getFaceCounts();
 		bRetrievalResult &= RetrieveTypedAbcData<Alembic::Abc::Int32ArraySamplePtr, uint32>(FaceCountsSample, FaceCounts);
+		if (!bRetrievalResult)
+		{
+			TSharedRef<FTokenizedMessage> Message = FTokenizedMessage::Create(EMessageSeverity::Error, LOCTEXT("NoFaceCounts", "Could not retrieve mesh face counts."));
+			FAbcImportLogger::AddImportMessage(Message);
+			return false;
+		}
 	}
 
 	const bool bNeedsTriangulation = FaceCounts.Contains(4);
@@ -494,6 +506,13 @@ bool AbcImporterUtilities::GenerateAbcMeshSampleDataForFrame(const Alembic::AbcG
 				}
 			}
 		}
+		else
+		{
+			TSharedRef<FTokenizedMessage> Message = FTokenizedMessage::Create(EMessageSeverity::Error, LOCTEXT("NoIndices", "Could not retrieve mesh indices."));
+			FAbcImportLogger::AddImportMessage(Message);
+			return false;
+		}
+
 		if (bNeedsTriangulation)
 		{
 			TriangulateIndexBuffer(FaceCounts, Sample->Indices);
diff --git a/Engine/Plugins/Importers/AlembicImporter/Source/AlembicLibrary/Private/AbcImportUtilities.h b/Engine/Plugins/Importers/AlembicImporter/Source/AlembicLibrary/Private/AbcImportUtilities.h
--- a/Engine/Plugins/Importers/AlembicImporter/Source/AlembicLibrary/Private/AbcImportUtilities.h
+++ b/Engine/Plugins/Importers/AlembicImporter/Source/AlembicLibrary/Private/AbcImportUtilities.h
@@ -70,6 +70,11 @@ namespace AbcImporterUtilities
 
 	template<typename T, typename U> bool RetrieveTypedAbcData(T InSampleDataPtr, TArray<U>& OutDataArray )
 	{
+		if (!InSampleDataPtr)
+		{
+			return false;
+		}
+
 		// Allocate required memory for the OutData
 		const int32 NumEntries = InSampleDataPtr->size();
 		bool bSuccess = false; 
