Fixed potential crash when importing Alembic with varying topology.

#rnx

diff --git a/Engine/Plugins/Importers/AlembicImporter/Source/AlembicLibrary/Private/AbcImportUtilities.cpp b/Engine/Plugins/Importers/AlembicImporter/Source/AlembicLibrary/Private/AbcImportUtilities.cpp
--- a/Engine/Plugins/Importers/AlembicImporter/Source/AlembicLibrary/Private/AbcImportUtilities.cpp
+++ b/Engine/Plugins/Importers/AlembicImporter/Source/AlembicLibrary/Private/AbcImportUtilities.cpp
@@ -1096,11 +1096,16 @@ FAbcMeshSample* AbcImporterUtilities::MergeMeshSamples(const TArray<const FAbcMe
 		MergedSample->MaterialIndices.Append(Sample->MaterialIndices);
 		MergedSample->SmoothingGroupIndices.Append(Sample->SmoothingGroupIndices);
 
-		// Remap material and smoothing group indices
+		// Remap material and smoothing group indices in separate loops since they might not have the same number of elements
 		const uint32 NumMaterialIndices = MergedSample->MaterialIndices.Num();
 		for (uint32 IndiceIndex = MaterialIndicesOffset; IndiceIndex < NumMaterialIndices; ++IndiceIndex)
 		{
 			MergedSample->MaterialIndices[IndiceIndex] += MergedSample->NumMaterials;
+		}
+
+		const uint32 NumSmoothingGroupIndices = MergedSample->SmoothingGroupIndices.Num();
+		for (uint32 IndiceIndex = SmoothingGroupIndicesOffset; IndiceIndex < NumSmoothingGroupIndices; ++IndiceIndex)
+		{
 			MergedSample->SmoothingGroupIndices[IndiceIndex] += MergedSample->NumSmoothingGroups;
 		}
 
@@ -1177,11 +1182,16 @@ void AbcImporterUtilities::AppendMeshSample(FAbcMeshSample* MeshSampleOne, const
 	MeshSampleOne->MaterialIndices.Append(MeshSampleTwo->MaterialIndices);
 	MeshSampleOne->SmoothingGroupIndices.Append(MeshSampleTwo->SmoothingGroupIndices);
 
-	// Remap material and smoothing group indices
+	// Remap material and smoothing group indices in separate loops since they might not have the same number of elements
 	const uint32 NumMaterialIndices = MeshSampleOne->MaterialIndices.Num();
 	for (uint32 IndiceIndex = MaterialIndicesOffset; IndiceIndex < NumMaterialIndices; ++IndiceIndex)
 	{
 		MeshSampleOne->MaterialIndices[IndiceIndex] += MeshSampleOne->NumMaterials;
+	}
+
+	const uint32 NumSmoothingGroupIndices = MeshSampleOne->SmoothingGroupIndices.Num();
+	for (uint32 IndiceIndex = SmoothingGroupIndicesOffset; IndiceIndex < NumSmoothingGroupIndices; ++IndiceIndex)
+	{
 		MeshSampleOne->SmoothingGroupIndices[IndiceIndex] += MeshSampleOne->NumSmoothingGroups;
 	}
 
@@ -1569,6 +1579,12 @@ void AbcImporterUtilities::GeometryCacheDataForMeshSample(FGeometryCacheMeshData
 	const uint32 MaterialOffset, const float SecondsPerFrame, const bool bUseVelocitiesAsMotionVectors, const bool bStoreImportedVertexNumbers)
 {
 	OutMeshData.BoundingBox = MeshSample->Vertices.Num() ? FBox3f(MeshSample->Vertices) : FBox3f(EForceInit::ForceInit);
+	if (MeshSample->Vertices.Num() == 0)
+	{
+		// Somehow the MeshSample can have Indices but no Vertices so just skip the whole thing if there are no vertices
+		// Only the BoundingBox needs to be initialized to avoid asserts when rendering
+		return;
+	}
 
 	// We currently always have everything except motion vectors
 	// and tangents are auto-configure based on their presence in the mesh sample data
