Alembic Import: Fixed out of bouds access (potentially due to negative times).

#Synced-CL 51095459

diff --git a/Engine/Plugins/Importers/AlembicImporter/Source/AlembicLibrary/Private/AbcImporter.cpp b/Engine/Plugins/Importers/AlembicImporter/Source/AlembicLibrary/Private/AbcImporter.cpp
--- a/Engine/Plugins/Importers/AlembicImporter/Source/AlembicLibrary/Private/AbcImporter.cpp
+++ b/Engine/Plugins/Importers/AlembicImporter/Source/AlembicLibrary/Private/AbcImporter.cpp
@@ -436,7 +436,7 @@ UGeometryCache* FAbcImporter::ImportAsGeometryCache(UObject* InParent, EObjectFl
 					bool bConstantTopology = true;
 					const bool bStoreImportedVertexNumbers = ImportSettings->GeometryCacheSettings.bStoreImportedVertexNumbers;
 
-					int32 FrameTimeIndex = FrameIndex - ImportSettings->SamplingSettings.FrameStart;
+					int32 FrameTimeIndex = FMath::Max(0, FrameIndex - ImportSettings->SamplingSettings.FrameStart);
 					AbcImporterUtilities::MergePolyMeshesToMeshData(FrameIndex, ImportSettings->SamplingSettings.FrameStart, AbcFile->GetSecondsPerFrame(), bUseVelocitiesAsMotionVectors,
 						PolyMeshes,
 						AbcFile->GetLookupMaterialSlot(),
