AlembicImporter: Fixed import failure when it couldn't retrieve velocities even though those should only be read when importing a geometry cache with the "Import Abc Velocities As Motion Vectors" option.

#Synced-CL 51237690

diff --git a/Engine/Plugins/Importers/AlembicImporter/Source/AlembicLibrary/Private/AbcImportUtilities.cpp b/Engine/Plugins/Importers/AlembicImporter/Source/AlembicLibrary/Private/AbcImportUtilities.cpp
--- a/Engine/Plugins/Importers/AlembicImporter/Source/AlembicLibrary/Private/AbcImportUtilities.cpp
+++ b/Engine/Plugins/Importers/AlembicImporter/Source/AlembicLibrary/Private/AbcImportUtilities.cpp
@@ -459,6 +459,12 @@ bool AbcImporterUtilities::GenerateAbcMeshSampleDataForFrame(const Alembic::AbcG
 		{
 			Alembic::Abc::V3fArraySamplePtr VelocitiesSample = MeshSample.getVelocities();
 			bRetrievalResult &= RetrieveTypedAbcData<Alembic::Abc::V3fArraySamplePtr, FVector3f>(VelocitiesSample, Sample->Velocities);
+			if (!bRetrievalResult)
+			{
+				TSharedRef<FTokenizedMessage> Message = FTokenizedMessage::Create(EMessageSeverity::Error, LOCTEXT("NoVelocities", "Could not retrieve velocities. Try importing without velocities as motion vectors."));
+				FAbcImportLogger::AddImportMessage(Message);
+				return false;
+			}
 		}
 	}
 
diff --git a/Engine/Plugins/Importers/AlembicImporter/Source/AlembicLibrary/Private/AbcPolyMesh.cpp b/Engine/Plugins/Importers/AlembicImporter/Source/AlembicLibrary/Private/AbcPolyMesh.cpp
--- a/Engine/Plugins/Importers/AlembicImporter/Source/AlembicLibrary/Private/AbcPolyMesh.cpp
+++ b/Engine/Plugins/Importers/AlembicImporter/Source/AlembicLibrary/Private/AbcPolyMesh.cpp
@@ -46,13 +46,27 @@ bool FAbcPolyMesh::ReadFirstFrame(const float InTime, const int32 FrameIndex)
 {
 	checkf(FirstSample == nullptr, TEXT("Reading First Frame Twice"));
 
+	// Only try to read velocities when importing a geometry cache with ImportAbcVelocitiesAsMotionVectors
+	const bool bReadVelocities = File->GetImportSettings()->ImportType == EAlembicImportType::GeometryCache &&
+		File->GetImportSettings()->GeometryCacheSettings.MotionVectors == EAbcGeometryCacheMotionVectorsImport::ImportAbcVelocitiesAsMotionVectors;
+
+	if (!bReadVelocities)
+	{
+		EnumRemoveFlags(SampleReadFlags, ESampleReadFlags::Velocities);
+	}
+
 	const float Time = InTime < MinTime ? MinTime : (InTime > MinTime ? MinTime : InTime);
 	Alembic::Abc::ISampleSelector SampleSelector = AbcImporterUtilities::GenerateAlembicSampleSelector<double>(Time);
 	const bool bFirstFrame = true;
 	FirstSample = AbcImporterUtilities::GenerateAbcMeshSampleForFrame(Schema, SampleSelector, SampleReadFlags, bFirstFrame);
 	bFirstFrameVisibility = AbcImporterUtilities::IsObjectVisible(PolyMesh, SampleSelector);
 	SampleReadFlags = AbcImporterUtilities::GenerateAbcMeshSampleReadFlags(Schema);
-	
+	// Remove the velocities flag again since GenerateAbcMeshSampleReadFlags generates a new SampleReadFlags from scratch
+	if (!bReadVelocities)
+	{
+		EnumRemoveFlags(SampleReadFlags, ESampleReadFlags::Velocities);
+	}
+
 	if (FirstSample)
 	{
 		if (FirstSample->Normals.Num() == 0)
